Refactored mesh merging dialogs into a non-modal panel in a separate Editor module. Moved actor merging out of 'experimental' classification.

diff --git a/Engine/Source/Editor/EditorStyle/Private/SlateEditorStyle.cpp b/Engine/Source/Editor/EditorStyle/Private/SlateEditorStyle.cpp
--- a/Engine/Source/Editor/EditorStyle/Private/SlateEditorStyle.cpp
+++ b/Engine/Source/Editor/EditorStyle/Private/SlateEditorStyle.cpp
@@ -2370,7 +2370,8 @@ void FSlateEditorStyle::FStyle::SetupGeneralStyles()
 		Set("StandardDialog.SlotPadding",FMargin(6.0f, 0.0f, 6.0f, 0.0f));
 		Set("StandardDialog.MinDesiredSlotWidth", 80.0f );
 		Set("StandardDialog.MinDesiredSlotHeight", 0.0f );
-		Set("StandardDialog.LargeFont",  TTF_CORE_FONT( "Fonts/Roboto-Regular", 11 ) );
+		Set("StandardDialog.SmallFont", TTF_CORE_FONT("Fonts/Roboto-Regular", 8));
+		Set("StandardDialog.LargeFont", TTF_CORE_FONT("Fonts/Roboto-Regular", 11));
 	}
 
 	// Highres Screenshot
@@ -4306,6 +4307,9 @@ void FSlateEditorStyle::FStyle::SetupLevelEditorStyle()
 		Set( "LevelEditor.OpenClassBlueprint", new IMAGE_BRUSH("Icons/icon_class_Blueprint_Open_16x", Icon16x16));
 		Set( "LevelEditor.EditMatinee", new IMAGE_BRUSH( "Icons/icon_matinee_40x", Icon40x40 ) );
 		Set( "LevelEditor.EditMatinee.Small", new IMAGE_BRUSH( "Icons/icon_matinee_40x", Icon20x20 ) );
+
+		Set( "MergeActors.MeshMergingTool", new IMAGE_BRUSH( "Icons/icon_MergeActors_MeshMerging_40x", Icon40x40 ) );
+		Set( "MergeActors.MeshProxyTool", new IMAGE_BRUSH( "Icons/icon_MergeActors_MeshProxy_40x", Icon40x40 ) );
 		
 		Set( "PlacementBrowser.OptionsMenu", new IMAGE_BRUSH( "Icons/icon_Blueprint_Macro_16x", Icon16x16 ) );
 
diff --git a/Engine/Source/Editor/LevelEditor/Private/LevelEditor.cpp b/Engine/Source/Editor/LevelEditor/Private/LevelEditor.cpp
--- a/Engine/Source/Editor/LevelEditor/Private/LevelEditor.cpp
+++ b/Engine/Source/Editor/LevelEditor/Private/LevelEditor.cpp
@@ -930,18 +930,6 @@ void FLevelEditorModule::BindGlobalLevelEditorCommands()
 		FExecuteAction::CreateStatic( &FLevelEditorActionCallbacks::RemoveActorsFromGroup_Clicked )
 		);
 	
-	ActionList.MapAction(
-		Commands.MergeActors,
-		FExecuteAction::CreateStatic( &FLevelEditorActionCallbacks::MergeActors_Clicked ),
-		FCanExecuteAction::CreateStatic( &FLevelEditorActionCallbacks::CanExecuteMergeActors )
-		);
-
-	ActionList.MapAction(
-		Commands.MergeActorsByMaterials,
-		FExecuteAction::CreateStatic( &FLevelEditorActionCallbacks::MergeActorsByMaterials_Clicked ),
-		FCanExecuteAction::CreateStatic( &FLevelEditorActionCallbacks::CanExecuteMergeActorsByMaterials )
-		);
-	
 	ActionList.MapAction(
 		Commands.ShowAll,
 		FExecuteAction::CreateStatic( &FLevelEditorActionCallbacks::ExecuteExecCommand, FString( TEXT("ACTOR UNHIDE ALL") ) ) 
diff --git a/Engine/Source/Editor/LevelEditor/Private/LevelEditorActions.cpp b/Engine/Source/Editor/LevelEditor/Private/LevelEditorActions.cpp
--- a/Engine/Source/Editor/LevelEditor/Private/LevelEditorActions.cpp
+++ b/Engine/Source/Editor/LevelEditor/Private/LevelEditorActions.cpp
@@ -343,7 +343,6 @@ bool FLevelEditorActionCallbacks::ToggleFavorite_IsChecked()
 	return bIsChecked;
 }
 
-
 bool FLevelEditorActionCallbacks::CanSaveWorld()
 {
 	return FSlateApplication::Get().IsNormalExecution() && (!GUnrealEd || !GUnrealEd->GetPackageAutoSaver().IsAutoSaving());
@@ -1682,34 +1681,6 @@ void FLevelEditorActionCallbacks::RemoveActorsFromGroup_Clicked()
 	GUnrealEd->edactRemoveFromGroup();
 }
 
-void FLevelEditorActionCallbacks::MergeActors_Clicked()
-{
-	GUnrealEd->edactMergeActors();
-}
-
-bool FLevelEditorActionCallbacks::CanExecuteMergeActors()
-{
-	IMeshUtilities* MeshUtilities = FModuleManager::Get().LoadModulePtr<IMeshUtilities>("MeshUtilities");
-	
-	if (MeshUtilities && MeshUtilities->GetMeshMergingInterface() != nullptr)
-	{
-		FSelectedActorInfo Info = AssetSelectionUtils::GetSelectedActorInfo();
-		return (Info.bHaveStaticMeshComponent || Info.bHaveLandscape);
-	}
-	
-	return false;
-}
-
-void FLevelEditorActionCallbacks::MergeActorsByMaterials_Clicked()
-{
-	GUnrealEd->edactMergeActorsByMaterials();
-}
-
-bool FLevelEditorActionCallbacks::CanExecuteMergeActorsByMaterials()
-{
-	FSelectedActorInfo Info = AssetSelectionUtils::GetSelectedActorInfo();
-	return Info.bHaveStaticMeshComponent;
-}
 
 void FLevelEditorActionCallbacks::LocationGridSnap_Clicked()
 {
@@ -2912,8 +2883,6 @@ void FLevelEditorCommands::RegisterCommands()
 	UI_COMMAND( RemoveActorsFromGroup, "Remove from Group", "Removes the selected actors from the selected groups", EUserInterfaceActionType::Button, FInputChord() );
 	UI_COMMAND( LockGroup, "Lock", "Locks the selected groups", EUserInterfaceActionType::Button, FInputChord() );
 	UI_COMMAND( UnlockGroup, "Unlock", "Unlocks the selected groups", EUserInterfaceActionType::Button, FInputChord() );
-	UI_COMMAND( MergeActors, "Create Mesh Proxy...", "Harvest geometry from selected actors and merge them into single mesh", EUserInterfaceActionType::Button, FInputChord() );
-	UI_COMMAND( MergeActorsByMaterials, "Merge Actors...", "Harvest geometry from selected actors and merge grouping them by materials", EUserInterfaceActionType::Button, FInputChord() );
 
 #if PLATFORM_MAC
 	UI_COMMAND( ShowAll, "Show All Actors", "Shows all actors", EUserInterfaceActionType::Button, FInputChord( EModifierKey::Command, EKeys::H ) );
diff --git a/Engine/Source/Editor/LevelEditor/Private/LevelEditorContextMenu.cpp b/Engine/Source/Editor/LevelEditor/Private/LevelEditorContextMenu.cpp
--- a/Engine/Source/Editor/LevelEditor/Private/LevelEditorContextMenu.cpp
+++ b/Engine/Source/Editor/LevelEditor/Private/LevelEditorContextMenu.cpp
@@ -109,13 +109,6 @@ public:
 	 */
 	static void FillEditMenu( class FMenuBuilder& MenuBuilder, LevelEditorMenuContext ContextType );
 
-	/**
-	 * Fills in menu options for the actors merging
-	 *
-	 * @param MenuBuilder	The menu to add items to
-	 */
-	static void FillMergeActorsMenu( class FMenuBuilder& MenuBuilder );
-
 private:
 	/**
 	 * Fills in menu options for the matinee selection menu
@@ -434,15 +427,6 @@ void FLevelEditorContextMenu::FillMenu( FMenuBuilder& MenuBuilder, TWeakPtr<SLev
 					MenuBuilder.AddMenuEntry(FLevelEditorCommands::Get().ResetPivot);
 				}
 			}
-
-			if (GetDefault<UEditorExperimentalSettings>()->bActorMerging &&
-				( SelectionInfo.bHaveStaticMeshComponent || SelectionInfo.bHaveLandscape ))
-			{
-				MenuBuilder.AddSubMenu(
-					LOCTEXT("MergeActorsSubMenu", "Merge"),
-					LOCTEXT("MergeActorsSubMenu_ToolTip", "Actor merging utils"),
-					FNewMenuDelegate::CreateStatic(&FLevelEditorContextMenuImpl::FillMergeActorsMenu));
-			}
 		}
 		MenuBuilder.EndSection();
 
@@ -1012,18 +996,6 @@ void FLevelEditorContextMenuImpl::FillEditMenu( class FMenuBuilder& MenuBuilder,
 	MenuBuilder.AddMenuEntry( FGenericCommands::Get().Rename );
 }
 
-void FLevelEditorContextMenuImpl::FillMergeActorsMenu( class FMenuBuilder& MenuBuilder )
-{
-	MenuBuilder.AddMenuEntry( FLevelEditorCommands::Get().MergeActorsByMaterials );
-
-	// Simplygon ProxyLOD
-	MenuBuilder.BeginSection("ProxySimplygon", LOCTEXT("SimplygonHeading", "Simplygon"));
-	{
-		MenuBuilder.AddMenuEntry(FLevelEditorCommands::Get().MergeActors);
-	}
-	MenuBuilder.EndSection();
-}
-
 void FLevelScriptEventMenuHelper::FillLevelBlueprintEventsMenu(class FMenuBuilder& MenuBuilder, const TArray<AActor*>& SelectedActors)
 {
 	AActor* SelectedActor = (1 == SelectedActors.Num()) ? SelectedActors[0] : NULL;
diff --git a/Engine/Source/Editor/LevelEditor/Public/LevelEditorActions.h b/Engine/Source/Editor/LevelEditor/Public/LevelEditorActions.h
--- a/Engine/Source/Editor/LevelEditor/Public/LevelEditorActions.h
+++ b/Engine/Source/Editor/LevelEditor/Public/LevelEditorActions.h
@@ -272,10 +272,6 @@ public:
 	TSharedPtr< FUICommandInfo > LockGroup;
 	/** Unlocks the selected group */
 	TSharedPtr< FUICommandInfo > UnlockGroup;
-	/** Opens a dialog window for creating mesh proxies */
-	TSharedPtr< FUICommandInfo > MergeActors;
-	/** Merge selected actors grouping them by materials */
-	TSharedPtr< FUICommandInfo > MergeActorsByMaterials;
 		
 	/**
 	 * Visibility commands                   
@@ -954,22 +950,6 @@ public:
 	 */
 	static void RemoveActorsFromGroup_Clicked();
 
-	/**
-	 * Called when the MergeActors command is executed
-	 */
-	static void MergeActors_Clicked();
-
-	/** @return Returns true if 'Merge Actors' can be used right now */
-	static bool CanExecuteMergeActors();
-
-	/**
-	 * Called when the MergeActorsByMaterials command is executed
-	 */
-	static void MergeActorsByMaterials_Clicked();
-
-	/** @return Returns true if 'Merge Actors y Materials' can be used right now */
-	static bool CanExecuteMergeActorsByMaterials();
-
 	/**
 	 * Called when the location grid snap is toggled off and on
 	 */
diff --git a/Engine/Source/Editor/MergeActors/MergeActors.Build.cs b/Engine/Source/Editor/MergeActors/MergeActors.Build.cs
new file mode 100644
--- /dev/null
+++ b/Engine/Source/Editor/MergeActors/MergeActors.Build.cs
@@ -0,0 +1,54 @@
+// Copyright 1998-2014 Epic Games, Inc. All Rights Reserved.
+
+using UnrealBuildTool;
+
+public class MergeActors : ModuleRules
+{
+	public MergeActors(TargetInfo Target)
+	{
+        PrivateIncludePaths.AddRange(
+            new string[] {
+				"Editor/MergeActors/Private",
+				"Editor/MergeActors/Private/MeshMergingTool",
+				"Editor/MergeActors/Private/MeshProxyTool"
+			}
+        );
+
+        PrivateIncludePathModuleNames.AddRange(
+			new string[] {
+				"AssetRegistry",
+                "ContentBrowser",
+                "Documentation",
+                "LevelEditor",
+                "MeshUtilities",
+                "PropertyEditor",
+                "RawMesh",
+                "WorkspaceMenuStructure"
+			}
+		);
+
+		PrivateDependencyModuleNames.AddRange(
+			new string[] {
+				"Core",
+				"CoreUObject",
+				"Engine",
+                "InputCore",
+				"Slate",
+				"SlateCore",
+                "EditorStyle",
+				"UnrealEd",
+			}
+		);
+
+		DynamicallyLoadedModuleNames.AddRange(
+			new string[] {
+				"AssetRegistry",
+                "ContentBrowser",
+                "Documentation",
+                "LevelEditor",
+                "MeshUtilities",
+                "PropertyEditor"
+			}
+		);
+	}
+}
diff --git a/Engine/Source/Editor/MergeActors/Private/MergeActorsModule.cpp b/Engine/Source/Editor/MergeActors/Private/MergeActorsModule.cpp
new file mode 100644
--- /dev/null
+++ b/Engine/Source/Editor/MergeActors/Private/MergeActorsModule.cpp
@@ -0,0 +1,153 @@
+// Copyright 1998-2014 Epic Games, Inc. All Rights Reserved.
+
+#include "MergeActorsPrivatePCH.h"
+#include "SMergeActorsToolbar.h"
+#include "ModuleManager.h"
+#include "WorkspaceMenuStructureModule.h"
+#include "MeshMergingTool.h"
+#include "MeshProxyTool.h"
+#include "SDockTab.h"
+
+#define LOCTEXT_NAMESPACE "MergeActorsModule"
+
+
+static const FName MergeActorsApp = FName("MergeActorsApp");
+
+/**
+ * Merge Actors module
+ */
+class FMergeActorsModule : public IMergeActorsModule
+{
+public:
+
+	/**
+	 * Called right after the module DLL has been loaded and the module object has been created
+	 */
+	virtual void StartupModule() override;
+
+	/**
+	 * Called before the module is unloaded, right before the module object is destroyed.
+	 */
+	virtual void ShutdownModule() override;
+
+	/**
+	 * Register an IMergeActorsTool with the module, passing ownership to it
+	 */
+	virtual bool RegisterMergeActorsTool(TUniquePtr<IMergeActorsTool> Tool) override;
+
+	/**
+	 * Unregister an IMergeActorsTool with the module
+	 */
+	virtual bool UnregisterMergeActorsTool(IMergeActorsTool* Tool) override;
+
+
+private:
+
+	/** Creates the dock tab widget used by the tab manager */
+	TSharedRef<SDockTab> CreateMergeActorsTab(const FSpawnTabArgs& Args);
+
+private:
+
+	/** Pointer to the toolbar widget */
+	TWeakPtr<SMergeActorsToolbar> MergeActorsToolbarPtr;
+
+	/** List of registered MergeActorsTool instances */
+	TArray<TUniquePtr<IMergeActorsTool>> MergeActorsTools;
+};
+
+IMPLEMENT_MODULE(FMergeActorsModule, MergeActors);
+
+
+TSharedRef<SDockTab> FMergeActorsModule::CreateMergeActorsTab(const FSpawnTabArgs& Args)
+{
+	// Build array of MergeActorsTool raw pointers
+	TArray<IMergeActorsTool*> ToolsToRegister;
+	for (const auto& MergeActorTool : MergeActorsTools)
+	{
+		check(MergeActorTool.Get() != nullptr);
+		ToolsToRegister.Add(MergeActorTool.Get());
+	}
+
+	// Construct toolbar widget
+	TSharedRef<SMergeActorsToolbar> MergeActorsToolbar =
+		SNew(SMergeActorsToolbar)
+		.ToolsToRegister(ToolsToRegister);
+
+	// Construct dock tab
+	TSharedRef<SDockTab> DockTab =
+		SNew(SDockTab)
+		.TabRole(ETabRole::NomadTab)
+		[
+			MergeActorsToolbar
+		];
+
+	// Keep weak pointer to toolbar widget
+	MergeActorsToolbarPtr = MergeActorsToolbar;
+
+	return DockTab;
+}
+
+
+void FMergeActorsModule::StartupModule()
+{
+	FGlobalTabmanager::Get()->RegisterNomadTabSpawner(MergeActorsApp, FOnSpawnTab::CreateRaw(this, &FMergeActorsModule::CreateMergeActorsTab))
+		.SetDisplayName(LOCTEXT("TabTitle", "Merge Actors"))
+		.SetTooltipText(LOCTEXT("TooltipText", "Open the Merge Actors tab."))
+		.SetGroup(WorkspaceMenu::GetMenuStructure().GetLevelEditorCategory())
+		.SetIcon(FSlateIcon(FEditorStyle::GetStyleSetName(), "ClassViewer.TabIcon"));
+
+	// Register built-in merging tools straight away
+	ensure(RegisterMergeActorsTool(MakeUnique<FMeshMergingTool>()));
+	ensure(RegisterMergeActorsTool(MakeUnique<FMeshProxyTool>()));
+}
+
+
+void FMergeActorsModule::ShutdownModule()
+{
+	if (FSlateApplication::IsInitialized())
+	{
+		FGlobalTabmanager::Get()->UnregisterNomadTabSpawner(MergeActorsApp);
+	}
+}
+
+
+bool FMergeActorsModule::RegisterMergeActorsTool(TUniquePtr<IMergeActorsTool> Tool)
+{
+	if (Tool.Get() != nullptr && !MergeActorsTools.Contains(MoveTemp(Tool)))
+	{
+		MergeActorsTools.Add(MoveTemp(Tool));
+		
+		// If a tool is added while the toolbar widget is active, update it accordingly
+		TSharedPtr<SMergeActorsToolbar> MergeActorsToolbar = MergeActorsToolbarPtr.Pin();
+		if (MergeActorsToolbar.IsValid())
+		{
+			MergeActorsToolbar->AddTool(Tool.Get());
+		}
+
+		return true;
+	}
+
+	return false;
+}
+
+
+bool FMergeActorsModule::UnregisterMergeActorsTool(IMergeActorsTool* Tool)
+{
+	if (Tool != nullptr)
+	{
+		if (MergeActorsTools.RemoveAll([=](const TUniquePtr<IMergeActorsTool>& Ptr) { return Ptr.Get() == Tool; }) > 0)
+		{
+			TSharedPtr<SMergeActorsToolbar> MergeActorsToolbar = MergeActorsToolbarPtr.Pin();
+			if (MergeActorsToolbar.IsValid())
+			{
+				MergeActorsToolbar->RemoveTool(Tool);
+			}
+
+			return true;
+		}
+	}
+	return false;
+}
+
+
+#undef LOCTEXT_NAMESPACE
diff --git a/Engine/Source/Editor/MergeActors/Private/MergeActorsPrivatePCH.h b/Engine/Source/Editor/MergeActors/Private/MergeActorsPrivatePCH.h
new file mode 100644
--- /dev/null
+++ b/Engine/Source/Editor/MergeActors/Private/MergeActorsPrivatePCH.h
@@ -0,0 +1,6 @@
+// Copyright 1998-2014 Epic Games, Inc. All Rights Reserved.
+
+#pragma once
+
+#include "UnrealEd.h"
+#include "IMergeActorsModule.h"
diff --git a/Engine/Source/Editor/MergeActors/Private/MeshMergingTool/MeshMergingTool.cpp b/Engine/Source/Editor/MergeActors/Private/MeshMergingTool/MeshMergingTool.cpp
new file mode 100644
--- /dev/null
+++ b/Engine/Source/Editor/MergeActors/Private/MeshMergingTool/MeshMergingTool.cpp
@@ -0,0 +1,141 @@
+// Copyright 1998-2014 Epic Games, Inc. All Rights Reserved.
+
+#include "MergeActorsPrivatePCH.h"
+#include "MeshMergingTool.h"
+#include "SMeshMergingDialog.h"
+#include "PropertyEditorModule.h"
+#include "Engine/TextureLODSettings.h"
+#include "RawMesh.h"
+#include "Engine/StaticMeshActor.h"
+#include "Engine/StaticMesh.h"
+#include "Engine/Selection.h"
+#include "SystemSettings.h"
+#include "Engine/TextureLODSettings.h"
+#include "ContentBrowserModule.h"
+#include "AssetRegistryModule.h"
+#include "ScopedTransaction.h"
+
+#define LOCTEXT_NAMESPACE "MeshMergingTool"
+
+
+FMeshMergingTool::FMeshMergingTool()
+	: bPlaceInWorld(false)
+	, bExportSpecificLOD(false)
+	, ExportLODIndex(0)
+{}
+
+
+TSharedRef<SWidget> FMeshMergingTool::GetWidget()
+{
+	return SNew(SMeshMergingDialog, this);
+}
+
+
+FText FMeshMergingTool::GetTooltipText() const
+{
+	return LOCTEXT("MeshMergingToolTooltip", "Harvest geometry from selected actors and merge grouping them by materials.");
+}
+
+
+FString FMeshMergingTool::GetDefaultPackageName() const
+{
+	FString PackageName = FPackageName::FilenameToLongPackageName(FPaths::GameContentDir() + TEXT("SM_MERGED"));
+
+	USelection* SelectedActors = GEditor->GetSelectedActors();
+	// Iterate through selected actors and find first static mesh asset
+	// Use this static mesh path as destination package name for a merged mesh
+	for (FSelectionIterator Iter(*SelectedActors); Iter; ++Iter)
+	{
+		AActor* Actor = Cast<AActor>(*Iter);
+		if (Actor)
+		{
+			FString ActorName = Actor->GetName();
+			PackageName = FString::Printf(TEXT("%s_%s"), *PackageName, *ActorName);
+			break;
+		}
+	}
+
+	if (PackageName.IsEmpty())
+	{
+		PackageName = MakeUniqueObjectName(NULL, UPackage::StaticClass(), *PackageName).ToString();
+	}
+
+	return PackageName;
+}
+
+
+bool FMeshMergingTool::RunMerge(const FString& PackageName)
+{
+	IMeshUtilities& MeshUtilities = FModuleManager::Get().LoadModuleChecked<IMeshUtilities>("MeshUtilities");
+	USelection* SelectedActors = GEditor->GetSelectedActors();
+	TArray<AActor*> Actors;
+	TArray<ULevel*> UniqueLevels;
+	for (FSelectionIterator Iter(*SelectedActors); Iter; ++Iter)
+	{
+		AActor* Actor = Cast<AActor>(*Iter);
+		if (Actor)
+		{
+			Actors.Add(Actor);
+			UniqueLevels.AddUnique(Actor->GetLevel());
+		}
+	}
+
+	// This restriction is only for replacement of selected actors with merged mesh actor
+	if (UniqueLevels.Num() > 1 && bPlaceInWorld)
+	{
+		FText Message = NSLOCTEXT("UnrealEd", "FailedToMergeActorsSublevels_Msg", "The selected actors should be in the same level");
+		OpenMsgDlgInt(EAppMsgType::Ok, Message, NSLOCTEXT("UnrealEd", "FailedToMergeActors_Title", "Unable to merge actors"));
+		return false;
+	}
+
+	int32 TargetMeshLOD = bExportSpecificLOD ? ExportLODIndex : INDEX_NONE;
+
+	FVector MergedActorLocation;
+	TArray<UObject*> AssetsToSync;
+	MeshUtilities.MergeActors(Actors, MergingSettings, NULL, PackageName, TargetMeshLOD, AssetsToSync, MergedActorLocation);
+
+	if (AssetsToSync.Num())
+	{
+		FAssetRegistryModule& AssetRegistry = FModuleManager::Get().LoadModuleChecked<FAssetRegistryModule>("AssetRegistry");
+		int32 AssetCount = AssetsToSync.Num();
+		for (int32 AssetIndex = 0; AssetIndex < AssetCount; AssetIndex++)
+		{
+			AssetRegistry.AssetCreated(AssetsToSync[AssetIndex]);
+			GEditor->BroadcastObjectReimported(AssetsToSync[AssetIndex]);
+		}
+
+		//Also notify the content browser that the new assets exists
+		FContentBrowserModule& ContentBrowserModule = FModuleManager::Get().LoadModuleChecked<FContentBrowserModule>("ContentBrowser");
+		ContentBrowserModule.Get().SyncBrowserToAssets(AssetsToSync, true);
+
+		// Place new mesh in the world
+		if (bPlaceInWorld)
+		{
+			const FScopedTransaction Transaction(NSLOCTEXT("UnrealEd", "PlaceMergedActor", "Place Merged Actor"));
+			UniqueLevels[0]->Modify();
+
+			UWorld* World = UniqueLevels[0]->OwningWorld;
+			FActorSpawnParameters Params;
+			Params.OverrideLevel = UniqueLevels[0];
+			FRotator MergedActorRotation(ForceInit);
+
+			AStaticMeshActor* MergedActor = World->SpawnActor<AStaticMeshActor>(MergedActorLocation, MergedActorRotation, Params);
+			MergedActor->GetStaticMeshComponent()->StaticMesh = Cast<UStaticMesh>(AssetsToSync[0]);
+			MergedActor->SetActorLabel(AssetsToSync[0]->GetName());
+
+			// Add source actors as children to merged actor and hide them
+			for (AActor* Actor : Actors)
+			{
+				Actor->Modify();
+				Actor->AttachRootComponentToActor(MergedActor, NAME_None, EAttachLocation::KeepWorldPosition);
+				Actor->SetActorHiddenInGame(true);
+				Actor->SetIsTemporarilyHiddenInEditor(true);
+			}
+		}
+	}
+
+	return true;
+}
+
+
+#undef LOCTEXT_NAMESPACE
diff --git a/Engine/Source/Editor/MergeActors/Private/MeshMergingTool/MeshMergingTool.h b/Engine/Source/Editor/MergeActors/Private/MeshMergingTool/MeshMergingTool.h
new file mode 100644
--- /dev/null
+++ b/Engine/Source/Editor/MergeActors/Private/MeshMergingTool/MeshMergingTool.h
@@ -0,0 +1,35 @@
+// Copyright 1998-2014 Epic Games, Inc. All Rights Reserved.
+#pragma once
+
+#include "IMergeActorsTool.h"
+#include "MeshUtilities.h"
+
+/**
+ * Mesh Merging Tool
+ */
+class FMeshMergingTool : public IMergeActorsTool
+{
+	friend class SMeshMergingDialog;
+
+public:
+
+	FMeshMergingTool();
+
+	// IMergeActorsTool interface
+	virtual TSharedRef<SWidget> GetWidget() override;
+	virtual FName GetIconName() const override { return "MergeActors.MeshMergingTool"; }
+	virtual FText GetTooltipText() const override;
+	virtual FString GetDefaultPackageName() const override;
+	virtual bool RunMerge(const FString& PackageName) override;
+
+private:
+
+	/** Current mesh merging settings */
+	FMeshMergingSettings MergingSettings;
+
+	/** Whether to spawn merged actor in the world */
+	bool bPlaceInWorld;
+
+	bool bExportSpecificLOD;
+	int32 ExportLODIndex;
+};
diff --git a/Engine/Source/Editor/MergeActors/Private/MeshMergingTool/SMeshMergingDialog.cpp b/Engine/Source/Editor/MergeActors/Private/MeshMergingTool/SMeshMergingDialog.cpp
new file mode 100644
--- /dev/null
+++ b/Engine/Source/Editor/MergeActors/Private/MeshMergingTool/SMeshMergingDialog.cpp
@@ -0,0 +1,490 @@
+// Copyright 1998-2015 Epic Games, Inc. All Rights Reserved.
+
+#include "MergeActorsPrivatePCH.h"
+#include "SMeshMergingDialog.h"
+#include "Dialogs/DlgPickAssetPath.h"
+#include "STextComboBox.h"
+#include "RawMesh.h"
+
+#define LOCTEXT_NAMESPACE "SMeshMergingDialog"
+
+//////////////////////////////////////////////////////////////////////////
+// SMeshMergingDialog
+
+SMeshMergingDialog::SMeshMergingDialog()
+{
+}
+
+BEGIN_SLATE_FUNCTION_BUILD_OPTIMIZATION
+void SMeshMergingDialog::Construct(const FArguments& InArgs, FMeshMergingTool* InTool)
+{
+	Tool = InTool;
+	check(Tool != nullptr);
+
+	// Setup available resolutions for lightmap and merged materials
+	const int32 MinTexResolution = 1 << FTextureLODGroup().MinLODMipCount;
+	const int32 MaxTexResolution = 1 << FTextureLODGroup().MaxLODMipCount;
+	for (int32 TexRes = MinTexResolution; TexRes <= MaxTexResolution; TexRes*=2)
+	{
+		LightMapResolutionOptions.Add(MakeShareable(new FString(FString::FormatAsNumber(TexRes))));
+		MergedMaterialResolutionOptions.Add(MakeShareable(new FString(FString::FormatAsNumber(TexRes))));
+	}
+
+	Tool->MergingSettings.TargetLightMapResolution = FMath::Clamp(Tool->MergingSettings.TargetLightMapResolution, MinTexResolution, MaxTexResolution);
+	Tool->MergingSettings.MergedMaterialAtlasResolution = FMath::Clamp(Tool->MergingSettings.MergedMaterialAtlasResolution, MinTexResolution, MaxTexResolution);
+		
+	// Setup available UV channels for an atlased lightmap
+	for (int32 Index = 0; Index < MAX_MESH_TEXTURE_COORDS; Index++)
+	{
+		LightMapChannelOptions.Add(MakeShareable(new FString(FString::FormatAsNumber(Index))));
+	}
+
+	for (int32 Index = 0; Index < MAX_STATIC_MESH_LODS; Index++)
+	{
+		ExportLODOptions.Add(MakeShareable(new FString(FString::FormatAsNumber(Index))));
+	}
+
+	// Create widget layout
+	this->ChildSlot
+	[
+		SNew(SVerticalBox)
+
+		+SVerticalBox::Slot()
+		.AutoHeight()
+		.Padding(0,2,0,0)
+		[
+			// Lightmap settings
+			SNew(SBorder)
+			.BorderImage(FEditorStyle::GetBrush("ToolPanel.GroupBorder"))
+			[
+				SNew(SVerticalBox)
+									
+				// Enable atlasing
+				+SVerticalBox::Slot()
+				.AutoHeight()
+				.Padding(FEditorStyle::GetMargin("StandardDialog.ContentPadding"))
+				[
+					SNew(SCheckBox)
+					.Type(ESlateCheckBoxType::CheckBox)
+					.IsChecked(this, &SMeshMergingDialog::GetGenerateLightmapUV)
+					.OnCheckStateChanged(this, &SMeshMergingDialog::SetGenerateLightmapUV)
+					.Content()
+					[
+						SNew(STextBlock)
+						.Text(LOCTEXT("AtlasLightmapUVLabel", "Generate Lightmap UVs"))
+						.Font(FEditorStyle::GetFontStyle("StandardDialog.SmallFont"))
+					]
+				]
+					
+				// Target lightmap channel / Max lightmap resolution
+				+SVerticalBox::Slot()
+				.AutoHeight()
+				.Padding(FEditorStyle::GetMargin("StandardDialog.ContentPadding"))
+				[
+					SNew(SHorizontalBox)
+
+					+SHorizontalBox::Slot()
+					.AutoWidth()
+					.VAlign(VAlign_Center)
+					[
+						SNew(STextBlock)
+						.IsEnabled(this, &SMeshMergingDialog::IsLightmapChannelEnabled)
+						.Text(LOCTEXT("TargetLightMapChannelLabel", "Target Channel:"))
+						.Font(FEditorStyle::GetFontStyle("StandardDialog.SmallFont"))
+					]
+
+					+SHorizontalBox::Slot()
+					.AutoWidth()
+					.VAlign(VAlign_Center)
+					.Padding(4,0,4,0)
+					[
+						SNew(STextComboBox)
+						.IsEnabled( this, &SMeshMergingDialog::IsLightmapChannelEnabled )
+						.OptionsSource(&LightMapChannelOptions)
+						.InitiallySelectedItem(LightMapChannelOptions[Tool->MergingSettings.TargetLightMapUVChannel])
+						.OnSelectionChanged(this, &SMeshMergingDialog::SetTargetLightMapChannel)
+						.Font(FEditorStyle::GetFontStyle("StandardDialog.SmallFont"))
+					]
+
+					+SHorizontalBox::Slot()
+					.AutoWidth()
+					.VAlign(VAlign_Center)
+					[
+						SNew(STextBlock)
+						.IsEnabled(this, &SMeshMergingDialog::IsLightmapChannelEnabled)
+						.Text(LOCTEXT("TargetLightMapResolutionLabel", "Target Resolution:"))
+						.Font(FEditorStyle::GetFontStyle("StandardDialog.SmallFont"))
+					]
+												
+					+SHorizontalBox::Slot()
+					.AutoWidth()
+					.VAlign(VAlign_Center)
+					.Padding(4,0,4,0)
+					[
+						SNew(STextComboBox)
+						.IsEnabled( this, &SMeshMergingDialog::IsLightmapChannelEnabled )
+						.OptionsSource(&LightMapResolutionOptions)
+						.InitiallySelectedItem(LightMapResolutionOptions[FMath::FloorLog2(Tool->MergingSettings.TargetLightMapResolution)])
+						.OnSelectionChanged(this, &SMeshMergingDialog::SetTargetLightMapResolution)
+						.Font(FEditorStyle::GetFontStyle("StandardDialog.SmallFont"))
+					]
+				]
+			]
+		]
+
+		// Other merging settings
+		+SVerticalBox::Slot()
+		.AutoHeight()
+		.Padding(0,2,0,0)
+		[
+			SNew(SBorder)
+			.BorderImage(FEditorStyle::GetBrush("ToolPanel.GroupBorder"))
+			[
+				SNew(SVerticalBox)
+
+				// LOD to export
+				+SVerticalBox::Slot()
+				.AutoHeight()
+				.Padding(FEditorStyle::GetMargin("StandardDialog.ContentPadding"))
+				[
+					SNew(SHorizontalBox)
+
+					+SHorizontalBox::Slot()
+					.AutoWidth()
+					.VAlign(VAlign_Center)
+					[
+						SNew(SCheckBox)
+						.Type(ESlateCheckBoxType::CheckBox)
+						.IsChecked(this, &SMeshMergingDialog::GetExportSpecificLODEnabled)
+						.OnCheckStateChanged(this, &SMeshMergingDialog::SetExportSpecificLODEnabled)
+						.Content()
+						[
+							SNew(STextBlock)
+							.IsEnabled(this, &SMeshMergingDialog::IsExportSpecificLODEnabled)
+							.Text(LOCTEXT("TargetMeshLODIndexLabel", "Export specific LOD:"))
+							.Font(FEditorStyle::GetFontStyle("StandardDialog.SmallFont"))
+						]
+					]
+					
+					+SHorizontalBox::Slot()
+					.AutoWidth()
+					.VAlign(VAlign_Center)
+					.Padding(4,0,4,0)
+					[
+						SNew(STextComboBox)
+						.IsEnabled(this, &SMeshMergingDialog::IsExportSpecificLODEnabled)
+						.OptionsSource(&ExportLODOptions)
+						.InitiallySelectedItem(ExportLODOptions[Tool->ExportLODIndex])
+						.OnSelectionChanged(this, &SMeshMergingDialog::SetExportSpecificLODIndex)
+						.Font(FEditorStyle::GetFontStyle("StandardDialog.SmallFont"))
+					]
+				]
+					
+				// Vertex colors
+				+SVerticalBox::Slot()
+				.AutoHeight()
+				.Padding(FEditorStyle::GetMargin("StandardDialog.ContentPadding"))
+				[
+					SNew(SCheckBox)
+					.Type(ESlateCheckBoxType::CheckBox)
+					.IsChecked(this, &SMeshMergingDialog::GetImportVertexColors)
+					.OnCheckStateChanged(this, &SMeshMergingDialog::SetImportVertexColors)
+					.Content()
+					[
+						SNew(STextBlock)
+						.Text(LOCTEXT("ImportVertexColorsLabel", "Import Vertex Colors"))
+						.Font(FEditorStyle::GetFontStyle("StandardDialog.SmallFont"))
+					]
+				]
+
+				// Pivot at zero
+				+SVerticalBox::Slot()
+				.AutoHeight()
+				.Padding(FEditorStyle::GetMargin("StandardDialog.ContentPadding"))
+				[
+					SNew(SCheckBox)
+					.Type(ESlateCheckBoxType::CheckBox)
+					.IsChecked(this, &SMeshMergingDialog::GetPivotPointAtZero)
+					.OnCheckStateChanged(this, &SMeshMergingDialog::SetPivotPointAtZero)
+					.Content()
+					[
+						SNew(STextBlock)
+						.Text(LOCTEXT("PivotPointAtZeroLabel", "Pivot Point At (0,0,0)"))
+						.Font(FEditorStyle::GetFontStyle("StandardDialog.SmallFont"))
+					]
+				]
+
+				// Place in world
+				+SVerticalBox::Slot()
+				.AutoHeight()
+				.Padding(FEditorStyle::GetMargin("StandardDialog.ContentPadding"))
+				[
+					SNew(SCheckBox)
+					.Type(ESlateCheckBoxType::CheckBox)
+					.IsChecked(this, &SMeshMergingDialog::GetPlaceInWorld)
+					.OnCheckStateChanged(this, &SMeshMergingDialog::SetPlaceInWorld)
+					.Content()
+					[
+						SNew(STextBlock)
+						.Text(LOCTEXT("PlaceInWorldLabel", "Place In World"))
+						.Font(FEditorStyle::GetFontStyle("StandardDialog.SmallFont"))
+					]
+				]
+
+				// Merge materials
+				+SVerticalBox::Slot()
+				.AutoHeight()
+				.Padding(FEditorStyle::GetMargin("StandardDialog.ContentPadding"))
+				[
+					SNew(SCheckBox)
+					.Type(ESlateCheckBoxType::CheckBox)
+					.IsChecked(this, &SMeshMergingDialog::GetMergeMaterials)
+					.OnCheckStateChanged(this, &SMeshMergingDialog::SetMergeMaterials)
+					.Content()
+					[
+						SNew(STextBlock)
+						.Text(LOCTEXT("MergeMaterialsLabel", "Merge Materials"))
+						.Font(FEditorStyle::GetFontStyle("StandardDialog.SmallFont"))
+					]
+				]
+				// Export normal map
+				+SVerticalBox::Slot()
+				.AutoHeight()
+				.Padding(FEditorStyle::GetMargin("StandardDialog.ContentPadding"))
+				[
+					SNew(SCheckBox)
+					.Type(ESlateCheckBoxType::CheckBox)
+					.IsChecked(this, &SMeshMergingDialog::GetExportNormalMap)
+					.OnCheckStateChanged(this, &SMeshMergingDialog::SetExportNormalMap)
+					.IsEnabled(this, &SMeshMergingDialog::IsMaterialMergingEnabled)
+					.Content()
+					[
+						SNew(STextBlock)
+						.Text(LOCTEXT("ExportNormalMapLabel", "Export Normal Map"))
+						.Font(FEditorStyle::GetFontStyle("StandardDialog.SmallFont"))
+					]
+				]
+				// Export metallic map
+				+SVerticalBox::Slot()
+				.AutoHeight()
+				.Padding(FEditorStyle::GetMargin("StandardDialog.ContentPadding"))
+				[
+					SNew(SCheckBox)
+					.Type(ESlateCheckBoxType::CheckBox)
+					.IsChecked(this, &SMeshMergingDialog::GetExportMetallicMap)
+					.OnCheckStateChanged(this, &SMeshMergingDialog::SetExportMetallicMap)
+					.IsEnabled(this, &SMeshMergingDialog::IsMaterialMergingEnabled)
+					.Content()
+					[
+						SNew(STextBlock)
+						.Text(LOCTEXT("ExportMetallicMapLabel", "Export Metallic Map"))
+						.Font(FEditorStyle::GetFontStyle("StandardDialog.SmallFont"))
+					]
+				]
+				// Export roughness map
+				+SVerticalBox::Slot()
+				.AutoHeight()
+				.Padding(FEditorStyle::GetMargin("StandardDialog.ContentPadding"))
+				[
+					SNew(SCheckBox)
+					.Type(ESlateCheckBoxType::CheckBox)
+					.IsChecked(this, &SMeshMergingDialog::GetExportRoughnessMap)
+					.OnCheckStateChanged(this, &SMeshMergingDialog::SetExportRoughnessMap)
+					.IsEnabled(this, &SMeshMergingDialog::IsMaterialMergingEnabled)
+					.Content()
+					[
+						SNew(STextBlock)
+						.Text(LOCTEXT("ExportRoughnessMapLabel", "Export Roughness Map"))
+						.Font(FEditorStyle::GetFontStyle("StandardDialog.SmallFont"))
+					]
+				]
+				// Export specular map
+				+SVerticalBox::Slot()
+				.AutoHeight()
+				.Padding(FEditorStyle::GetMargin("StandardDialog.ContentPadding"))
+				[
+					SNew(SCheckBox)
+					.Type(ESlateCheckBoxType::CheckBox)
+					.IsChecked(this, &SMeshMergingDialog::GetExportSpecularMap)
+					.OnCheckStateChanged(this, &SMeshMergingDialog::SetExportSpecularMap)
+					.IsEnabled(this, &SMeshMergingDialog::IsMaterialMergingEnabled)
+					.Content()
+					[
+						SNew(STextBlock)
+						.Text(LOCTEXT("ExportSpecularMapLabel", "Export Specular Map"))
+						.Font(FEditorStyle::GetFontStyle("StandardDialog.SmallFont"))
+					]
+				]
+				
+				// Merged texture size
+				+SVerticalBox::Slot()
+				.AutoHeight()
+				.Padding(FEditorStyle::GetMargin("StandardDialog.ContentPadding"))
+				[
+					SNew(SHorizontalBox)
+					+SHorizontalBox::Slot()
+					.AutoWidth()
+					.VAlign(VAlign_Center)
+					[
+						SNew(STextBlock)
+						.IsEnabled(this, &SMeshMergingDialog::IsMaterialMergingEnabled)
+						.Text(LOCTEXT("MergedMaterialAtlasResolutionLabel", "Merged Material Atlas Resolution:"))
+						.Font(FEditorStyle::GetFontStyle("StandardDialog.SmallFont"))
+					]
+												
+					+SHorizontalBox::Slot()
+					.AutoWidth()
+					.VAlign(VAlign_Center)
+					.Padding(4,0,4,0)
+					[
+						SNew(STextComboBox)
+						.IsEnabled(this, &SMeshMergingDialog::IsMaterialMergingEnabled)
+						.OptionsSource(&MergedMaterialResolutionOptions)
+						.InitiallySelectedItem(MergedMaterialResolutionOptions[FMath::FloorLog2(Tool->MergingSettings.MergedMaterialAtlasResolution)])
+						.OnSelectionChanged(this, &SMeshMergingDialog::SetMergedMaterialAtlasResolution)
+						.Font(FEditorStyle::GetFontStyle("StandardDialog.SmallFont"))
+					]
+				]														
+			]
+		]
+	];
+}
+END_SLATE_FUNCTION_BUILD_OPTIMIZATION
+
+ECheckBoxState SMeshMergingDialog::GetGenerateLightmapUV() const
+{
+	return (Tool->MergingSettings.bGenerateLightMapUV ? ECheckBoxState::Checked : ECheckBoxState::Unchecked);
+}
+
+void SMeshMergingDialog::SetGenerateLightmapUV(ECheckBoxState NewValue)
+{
+	Tool->MergingSettings.bGenerateLightMapUV = (ECheckBoxState::Checked == NewValue);
+}
+
+bool SMeshMergingDialog::IsLightmapChannelEnabled() const
+{
+	return Tool->MergingSettings.bGenerateLightMapUV;
+}
+
+void SMeshMergingDialog::SetTargetLightMapChannel(TSharedPtr<FString> NewSelection, ESelectInfo::Type SelectInfo)
+{
+	TTypeFromString<int32>::FromString(Tool->MergingSettings.TargetLightMapUVChannel, **NewSelection);
+}
+
+void SMeshMergingDialog::SetTargetLightMapResolution(TSharedPtr<FString> NewSelection, ESelectInfo::Type SelectInfo)
+{
+	TTypeFromString<int32>::FromString(Tool->MergingSettings.TargetLightMapResolution, **NewSelection);
+}
+
+ECheckBoxState SMeshMergingDialog::GetExportSpecificLODEnabled() const
+{
+	return (Tool->bExportSpecificLOD ? ECheckBoxState::Checked : ECheckBoxState::Unchecked);
+}
+
+void SMeshMergingDialog::SetExportSpecificLODEnabled(ECheckBoxState NewValue)
+{
+	Tool->bExportSpecificLOD = (NewValue == ECheckBoxState::Checked);
+}
+
+bool SMeshMergingDialog::IsExportSpecificLODEnabled() const
+{
+	return Tool->bExportSpecificLOD;
+}
+
+void SMeshMergingDialog::SetExportSpecificLODIndex(TSharedPtr<FString> NewSelection, ESelectInfo::Type SelectInfo)
+{
+	TTypeFromString<int32>::FromString(Tool->ExportLODIndex, **NewSelection);
+}
+
+ECheckBoxState SMeshMergingDialog::GetImportVertexColors() const
+{
+	return (Tool->MergingSettings.bImportVertexColors ? ECheckBoxState::Checked : ECheckBoxState::Unchecked);
+}
+
+void SMeshMergingDialog::SetImportVertexColors(ECheckBoxState NewValue)
+{
+	Tool->MergingSettings.bImportVertexColors = (ECheckBoxState::Checked == NewValue);
+}
+
+ECheckBoxState SMeshMergingDialog::GetPivotPointAtZero() const
+{
+	return (Tool->MergingSettings.bPivotPointAtZero ? ECheckBoxState::Checked : ECheckBoxState::Unchecked);
+}
+
+void SMeshMergingDialog::SetPivotPointAtZero(ECheckBoxState NewValue)
+{
+	Tool->MergingSettings.bPivotPointAtZero = (ECheckBoxState::Checked == NewValue);
+}
+
+ECheckBoxState SMeshMergingDialog::GetPlaceInWorld() const
+{
+	return (Tool->bPlaceInWorld ? ECheckBoxState::Checked : ECheckBoxState::Unchecked);
+}
+
+void SMeshMergingDialog::SetPlaceInWorld(ECheckBoxState NewValue)
+{
+	Tool->bPlaceInWorld = (ECheckBoxState::Checked == NewValue);
+}
+
+bool SMeshMergingDialog::IsMaterialMergingEnabled() const
+{
+	return Tool->MergingSettings.bMergeMaterials;
+}
+
+ECheckBoxState SMeshMergingDialog::GetMergeMaterials() const
+{
+	return (Tool->MergingSettings.bMergeMaterials ? ECheckBoxState::Checked : ECheckBoxState::Unchecked);
+}
+
+void SMeshMergingDialog::SetMergeMaterials(ECheckBoxState NewValue)
+{
+	Tool->MergingSettings.bMergeMaterials = (ECheckBoxState::Checked == NewValue);
+}
+
+ECheckBoxState SMeshMergingDialog::GetExportNormalMap() const
+{
+	return (Tool->MergingSettings.bExportNormalMap ? ECheckBoxState::Checked : ECheckBoxState::Unchecked);
+}
+
+void SMeshMergingDialog::SetExportNormalMap(ECheckBoxState NewValue)
+{
+	Tool->MergingSettings.bExportNormalMap = (ECheckBoxState::Checked == NewValue);
+}
+
+ECheckBoxState SMeshMergingDialog::GetExportMetallicMap() const
+{
+	return (Tool->MergingSettings.bExportMetallicMap ? ECheckBoxState::Checked : ECheckBoxState::Unchecked);
+}
+
+void SMeshMergingDialog::SetExportMetallicMap(ECheckBoxState NewValue)
+{
+	Tool->MergingSettings.bExportMetallicMap = (ECheckBoxState::Checked == NewValue);
+}
+
+ECheckBoxState SMeshMergingDialog::GetExportRoughnessMap() const
+{
+	return (Tool->MergingSettings.bExportRoughnessMap ? ECheckBoxState::Checked : ECheckBoxState::Unchecked);
+}
+
+void SMeshMergingDialog::SetExportRoughnessMap(ECheckBoxState NewValue)
+{
+	Tool->MergingSettings.bExportRoughnessMap = (ECheckBoxState::Checked == NewValue);
+}
+
+ECheckBoxState SMeshMergingDialog::GetExportSpecularMap() const
+{
+	return (Tool->MergingSettings.bExportSpecularMap ? ECheckBoxState::Checked : ECheckBoxState::Unchecked);
+}
+
+void SMeshMergingDialog::SetExportSpecularMap(ECheckBoxState NewValue)
+{
+	Tool->MergingSettings.bExportSpecularMap = (ECheckBoxState::Checked == NewValue);
+}
+
+void SMeshMergingDialog::SetMergedMaterialAtlasResolution(TSharedPtr<FString> NewSelection, ESelectInfo::Type SelectInfo)
+{
+	TTypeFromString<int32>::FromString(Tool->MergingSettings.MergedMaterialAtlasResolution, **NewSelection);
+}
+
+
+#undef LOCTEXT_NAMESPACE
diff --git a/Engine/Source/Editor/MergeActors/Private/MeshMergingTool/SMeshMergingDialog.h b/Engine/Source/Editor/MergeActors/Private/MeshMergingTool/SMeshMergingDialog.h
new file mode 100644
--- /dev/null
+++ b/Engine/Source/Editor/MergeActors/Private/MeshMergingTool/SMeshMergingDialog.h
@@ -0,0 +1,83 @@
+// Copyright 1998-2015 Epic Games, Inc. All Rights Reserved.
+#pragma once
+
+class FMeshMergingTool;
+
+/*-----------------------------------------------------------------------------
+   SMeshMergingDialog
+-----------------------------------------------------------------------------*/
+class SMeshMergingDialog : public SCompoundWidget
+{
+public:
+	SLATE_BEGIN_ARGS(SMeshMergingDialog)
+	{
+	}
+
+	SLATE_END_ARGS()
+
+public:
+	/** **/
+	SMeshMergingDialog();
+
+	/** SWidget functions */
+	void Construct(const FArguments& InArgs, FMeshMergingTool* InTool);
+
+private:
+
+	/** Called when the Merge button is clicked */
+	FReply OnMergeClicked();
+
+	/**  */
+	ECheckBoxState GetGenerateLightmapUV() const;
+	void SetGenerateLightmapUV(ECheckBoxState NewValue);
+
+	/** Target lightmap channel */
+	bool IsLightmapChannelEnabled() const;
+	void SetTargetLightMapChannel(TSharedPtr<FString> NewSelection, ESelectInfo::Type SelectInfo);
+	void SetTargetLightMapResolution(TSharedPtr<FString> NewSelection, ESelectInfo::Type SelectInfo);
+
+	/**  */
+	ECheckBoxState GetExportSpecificLODEnabled() const;
+	void SetExportSpecificLODEnabled(ECheckBoxState NewValue);
+	bool IsExportSpecificLODEnabled() const;
+	void SetExportSpecificLODIndex(TSharedPtr<FString> NewSelection, ESelectInfo::Type SelectInfo);
+		
+	/**  */
+	ECheckBoxState GetImportVertexColors() const;
+	void SetImportVertexColors(ECheckBoxState NewValue);
+
+	/**  */
+	ECheckBoxState GetPivotPointAtZero() const;
+	void SetPivotPointAtZero(ECheckBoxState NewValue);
+
+	ECheckBoxState GetPlaceInWorld() const;
+	void SetPlaceInWorld(ECheckBoxState NewValue);
+
+	/** Material merging */
+	bool IsMaterialMergingEnabled() const;
+	ECheckBoxState GetMergeMaterials() const;
+	void SetMergeMaterials(ECheckBoxState NewValue);
+
+	ECheckBoxState GetExportNormalMap() const;
+	void SetExportNormalMap(ECheckBoxState NewValue);
+
+	ECheckBoxState GetExportMetallicMap() const;
+	void SetExportMetallicMap(ECheckBoxState NewValue);
+
+	ECheckBoxState GetExportRoughnessMap() const;
+	void SetExportRoughnessMap(ECheckBoxState NewValue);
+
+	ECheckBoxState GetExportSpecularMap() const;
+	void SetExportSpecularMap(ECheckBoxState NewValue);
+
+	void SetMergedMaterialAtlasResolution(TSharedPtr<FString> NewSelection, ESelectInfo::Type SelectInfo);
+	
+private:
+
+	FMeshMergingTool* Tool;
+
+	TArray<TSharedPtr<FString>>	ExportLODOptions;
+	TArray<TSharedPtr<FString>>	LightMapResolutionOptions;
+	TArray<TSharedPtr<FString>>	LightMapChannelOptions;
+	TArray<TSharedPtr<FString>>	MergedMaterialResolutionOptions;
+};
diff --git a/Engine/Source/Editor/MergeActors/Private/MeshProxyTool/MeshProxyTool.cpp b/Engine/Source/Editor/MergeActors/Private/MeshProxyTool/MeshProxyTool.cpp
new file mode 100644
--- /dev/null
+++ b/Engine/Source/Editor/MergeActors/Private/MeshProxyTool/MeshProxyTool.cpp
@@ -0,0 +1,121 @@
+// Copyright 1998-2014 Epic Games, Inc. All Rights Reserved.
+
+#include "MergeActorsPrivatePCH.h"
+#include "MeshProxyTool.h"
+#include "SMeshProxyDialog.h"
+#include "MeshUtilities.h"
+#include "ContentBrowserModule.h"
+#include "AssetRegistryModule.h"
+#include "Engine/StaticMeshActor.h"
+#include "Engine/StaticMesh.h"
+#include "Engine/Selection.h"
+
+
+#define LOCTEXT_NAMESPACE "MeshProxyTool"
+
+
+TSharedRef<SWidget> FMeshProxyTool::GetWidget()
+{
+	return SNew(SMeshProxyDialog, this);
+}
+
+
+FText FMeshProxyTool::GetTooltipText() const
+{
+	return LOCTEXT("MeshProxyToolTooltip", "Harvest geometry from selected actors and merge them into single mesh.");
+}
+
+
+FString FMeshProxyTool::GetDefaultPackageName() const
+{
+	FString PackageName;
+
+	USelection* SelectedActors = GEditor->GetSelectedActors();
+
+	// Iterate through selected actors and find first static mesh asset
+	// Use this static mesh path as destination package name for a merged mesh
+	for (FSelectionIterator Iter(*SelectedActors); Iter; ++Iter)
+	{
+		AActor* Actor = Cast<AActor>(*Iter);
+		if (Actor)
+		{
+			TInlineComponentArray<UStaticMeshComponent*> SMComponets;
+			Actor->GetComponents<UStaticMeshComponent>(SMComponets);
+			for (UStaticMeshComponent* Component : SMComponets)
+			{
+				if (Component->StaticMesh)
+				{
+					PackageName = FPackageName::GetLongPackagePath(Component->StaticMesh->GetOutermost()->GetName());
+					PackageName += FString(TEXT("/PROXY_")) + Component->StaticMesh->GetName();
+					break;
+				}
+			}
+		}
+
+		if (!PackageName.IsEmpty())
+		{
+			break;
+		}
+	}
+
+	if (PackageName.IsEmpty())
+	{
+		PackageName = FPackageName::FilenameToLongPackageName(FPaths::GameContentDir() + TEXT("PROXY"));
+		PackageName = MakeUniqueObjectName(NULL, UPackage::StaticClass(), *PackageName).ToString();
+	}
+
+	return PackageName;
+}
+
+
+bool FMeshProxyTool::RunMerge(const FString& PackageName)
+{
+	TArray<AActor*> Actors;
+	TArray<UObject*> AssetsToSync;
+
+	//Get the module for the proxy mesh utilities
+	IMeshUtilities& MeshUtilities = FModuleManager::Get().LoadModuleChecked<IMeshUtilities>("MeshUtilities");
+
+	USelection* SelectedActors = GEditor->GetSelectedActors();
+	for (FSelectionIterator Iter(*SelectedActors); Iter; ++Iter)
+	{
+		AActor* Actor = Cast<AActor>(*Iter);
+		if (Actor)
+		{
+			Actors.Add(Actor);
+		}
+	}
+
+	if (Actors.Num())
+	{
+		GWarn->BeginSlowTask(LOCTEXT("MeshProxy_CreatingProxy", "Creating Mesh Proxy"), true);
+		GEditor->BeginTransaction(LOCTEXT("MeshProxy_Create", "Creating Mesh Proxy"));
+
+		FVector ProxyLocation = FVector::ZeroVector;
+		MeshUtilities.CreateProxyMesh(Actors, ProxySettings, NULL, PackageName, AssetsToSync, ProxyLocation);
+
+		GEditor->EndTransaction();
+		GWarn->EndSlowTask();
+
+		//Update the asset registry that a new static mash and material has been created
+		if (AssetsToSync.Num())
+		{
+			FAssetRegistryModule& AssetRegistry = FModuleManager::Get().LoadModuleChecked<FAssetRegistryModule>("AssetRegistry");
+			int32 AssetCount = AssetsToSync.Num();
+			for (int32 AssetIndex = 0; AssetIndex < AssetCount; AssetIndex++)
+			{
+				AssetRegistry.AssetCreated(AssetsToSync[AssetIndex]);
+				GEditor->BroadcastObjectReimported(AssetsToSync[AssetIndex]);
+			}
+
+			//Also notify the content browser that the new assets exists
+			FContentBrowserModule& ContentBrowserModule = FModuleManager::Get().LoadModuleChecked<FContentBrowserModule>("ContentBrowser");
+			ContentBrowserModule.Get().SyncBrowserToAssets(AssetsToSync, true);
+		}
+	}
+
+	return true;
+}
+
+
+#undef LOCTEXT_NAMESPACE
diff --git a/Engine/Source/Editor/MergeActors/Private/MeshProxyTool/MeshProxyTool.h b/Engine/Source/Editor/MergeActors/Private/MeshProxyTool/MeshProxyTool.h
new file mode 100644
--- /dev/null
+++ b/Engine/Source/Editor/MergeActors/Private/MeshProxyTool/MeshProxyTool.h
@@ -0,0 +1,25 @@
+// Copyright 1998-2014 Epic Games, Inc. All Rights Reserved.
+#pragma once
+
+#include "IMergeActorsTool.h"
+
+/**
+* Mesh Proxy Tool
+*/
+class FMeshProxyTool : public IMergeActorsTool
+{
+	friend class SMeshProxyDialog;
+
+public:
+
+	// IMergeActorsTool interface
+	virtual TSharedRef<SWidget> GetWidget() override;
+	virtual FName GetIconName() const override { return "MergeActors.MeshProxyTool"; }
+	virtual FText GetTooltipText() const override;
+	virtual FString GetDefaultPackageName() const override;
+	virtual bool RunMerge(const FString& PackageName) override;
+
+private:
+
+	FMeshProxySettings ProxySettings;
+};
diff --git a/Engine/Source/Editor/MergeActors/Private/MeshProxyTool/SMeshProxyDialog.cpp b/Engine/Source/Editor/MergeActors/Private/MeshProxyTool/SMeshProxyDialog.cpp
new file mode 100644
--- /dev/null
+++ b/Engine/Source/Editor/MergeActors/Private/MeshProxyTool/SMeshProxyDialog.cpp
@@ -0,0 +1,560 @@
+// Copyright 1998-2015 Epic Games, Inc. All Rights Reserved.
+
+#include "MergeActorsPrivatePCH.h"
+#include "SMeshProxyDialog.h"
+#include "Dialogs/DlgPickAssetPath.h"
+#include "SNumericEntryBox.h"
+#include "STextComboBox.h"
+
+#define LOCTEXT_NAMESPACE "SMeshProxyDialog"
+
+
+void SMeshProxyDialog::Construct(const FArguments& InArgs, FMeshProxyTool* InTool)
+{
+	Tool = InTool;
+	check(Tool != nullptr);
+
+	CuttingPlaneOptions.Add(MakeShareable(new FString(TEXT("+X"))));
+	CuttingPlaneOptions.Add(MakeShareable(new FString(TEXT("+Y"))));
+	CuttingPlaneOptions.Add(MakeShareable(new FString(TEXT("+Z"))));
+	CuttingPlaneOptions.Add(MakeShareable(new FString(TEXT("-X"))));
+	CuttingPlaneOptions.Add(MakeShareable(new FString(TEXT("-Y"))));
+	CuttingPlaneOptions.Add(MakeShareable(new FString(TEXT("-Z"))));
+
+	TextureResolutionOptions.Add(MakeShareable(new FString(TEXT("64"))));
+	TextureResolutionOptions.Add(MakeShareable(new FString(TEXT("128"))));
+	TextureResolutionOptions.Add(MakeShareable(new FString(TEXT("256"))));
+	TextureResolutionOptions.Add(MakeShareable(new FString(TEXT("512"))));
+	TextureResolutionOptions.Add(MakeShareable(new FString(TEXT("1024"))));
+	TextureResolutionOptions.Add(MakeShareable(new FString(TEXT("2048"))));
+
+	CreateLayout();
+}
+
+BEGIN_SLATE_FUNCTION_BUILD_OPTIMIZATION
+void SMeshProxyDialog::CreateLayout()
+{
+	this->ChildSlot
+	[
+		SNew(SVerticalBox)
+			
+		+SVerticalBox::Slot()
+		.AutoHeight()
+		.HAlign(HAlign_Center)
+		.VAlign(VAlign_Center)
+		.Padding(10)
+		[
+			// Simplygon logo
+			SNew(SImage)
+			.Image(FEditorStyle::GetBrush("MeshProxy.SimplygonLogo"))
+		]
+			
+		+SVerticalBox::Slot()
+		.AutoHeight()
+		[
+			SNew(SBorder)
+			.BorderImage(FEditorStyle::GetBrush("ToolPanel.GroupBorder"))
+			[
+				// Proxy options
+				SNew(SVerticalBox)
+
+				+SVerticalBox::Slot()
+				.AutoHeight()
+				.Padding(FEditorStyle::GetMargin("StandardDialog.ContentPadding"))
+				[
+					SNew(SHorizontalBox)
+					+SHorizontalBox::Slot()
+					.FillWidth(0.5f)
+					.VAlign(VAlign_Center)
+					[
+						SNew(STextBlock)
+						.Text(LOCTEXT("OnScreenSizeLabel", "On Screen Size (pixels)"))
+						.Font(FEditorStyle::GetFontStyle("StandardDialog.SmallFont"))
+					]
+					+ SHorizontalBox::Slot()
+					.FillWidth(0.5f)
+					.HAlign(HAlign_Left)
+					.VAlign(VAlign_Center)
+					[
+						SNew(SBox)
+						.HAlign(HAlign_Fill)
+						.MinDesiredWidth(100.0f)
+						.MaxDesiredWidth(100.0f)
+						[
+							SNew(SNumericEntryBox<int32>)
+							.Font(FEditorStyle::GetFontStyle("StandardDialog.SmallFont"))
+							.MinValue(40)
+							.MaxValue(1200)
+							.MinSliderValue(40)
+							.MaxSliderValue(1200)
+							.AllowSpin(true)
+							.Value(this, &SMeshProxyDialog::GetScreenSize)
+							.OnValueChanged(this, &SMeshProxyDialog::ScreenSizeChanged)
+						]
+					]
+				]
+
+				+SVerticalBox::Slot()
+				.AutoHeight()
+				.Padding(FEditorStyle::GetMargin("StandardDialog.ContentPadding"))
+				[
+					SNew(SHorizontalBox)
+					+SHorizontalBox::Slot()
+					.FillWidth(0.5f)
+					.VAlign(VAlign_Center)
+					[
+						SNew(STextBlock)
+						.Text(LOCTEXT("MergeDistanceLabel", "Merge Distance (pixels)"))
+						.Font(FEditorStyle::GetFontStyle("StandardDialog.SmallFont"))
+					]
+					+ SHorizontalBox::Slot()
+					.FillWidth(0.5f)
+					.HAlign(HAlign_Left)
+					.VAlign(VAlign_Center)
+					[
+						SNew(SBox)
+						.HAlign(HAlign_Fill)
+						.MinDesiredWidth(100.0f)
+						.MaxDesiredWidth(100.0f)
+						[
+							SNew(SNumericEntryBox<int32>)
+							.Font(FEditorStyle::GetFontStyle("StandardDialog.SmallFont"))
+							.MinValue(0)
+							.MaxValue(300)
+							.MinSliderValue(0)
+							.MaxSliderValue(300)
+							.AllowSpin(true)
+							.Value(this, &SMeshProxyDialog::GetMergeDistance)
+							.OnValueChanged(this, &SMeshProxyDialog::MergeDistanceChanged)
+						]
+					]
+				]
+
+				+SVerticalBox::Slot()
+				.AutoHeight()
+				.Padding(FEditorStyle::GetMargin("StandardDialog.ContentPadding"))
+				[
+					SNew(SHorizontalBox)
+					+SHorizontalBox::Slot()
+					.FillWidth(0.5f)
+					.VAlign(VAlign_Center)
+					[
+						SNew(STextBlock)
+						.Text(LOCTEXT("TextureResolutionLabel", "Texture Resolution"))
+						.Font(FEditorStyle::GetFontStyle("StandardDialog.SmallFont"))
+					]
+					+ SHorizontalBox::Slot()
+					.FillWidth(0.5f)
+					.HAlign(HAlign_Left)
+					.VAlign(VAlign_Center)
+					[
+						SNew(STextComboBox)
+						.Font(FEditorStyle::GetFontStyle("StandardDialog.SmallFont"))
+						.OptionsSource(&TextureResolutionOptions)
+						.InitiallySelectedItem(TextureResolutionOptions[3]) //512
+						.OnSelectionChanged(this, &SMeshProxyDialog::SetTextureResolution)
+					]
+				]
+
+				+SVerticalBox::Slot()
+				.AutoHeight()
+				.Padding(FEditorStyle::GetMargin("StandardDialog.ContentPadding"))
+				[
+					SNew(SCheckBox)
+					.Type(ESlateCheckBoxType::CheckBox)
+					.IsChecked(this, &SMeshProxyDialog::GetRecalculateNormals)
+					.OnCheckStateChanged(this, &SMeshProxyDialog::SetRecalculateNormals)
+					.Content()
+					[
+						SNew(STextBlock)
+						.Text(LOCTEXT("RecalcNormalsLabel", "Recalculate Normals"))
+						.Font(FEditorStyle::GetFontStyle("StandardDialog.SmallFont"))
+					]
+				]
+
+				+SVerticalBox::Slot()
+				.AutoHeight()
+				.Padding(FEditorStyle::GetMargin("StandardDialog.ContentPadding"))
+				[
+					SNew(SHorizontalBox)
+					+SHorizontalBox::Slot()
+					.FillWidth(0.5f)
+					.VAlign(VAlign_Center)
+					.Padding(0.0, 0.0, 3.0, 0.0)
+					[
+						SNew(SHorizontalBox)
+						+SHorizontalBox::Slot()
+						.AutoWidth()
+						[
+							SNew(SBox)
+							.MinDesiredWidth(30)
+							[
+								SNullWidget::NullWidget
+							]
+						]
+						+SHorizontalBox::Slot()
+						.AutoWidth()
+						.HAlign(HAlign_Left)
+						[
+							SNew(STextBlock)
+							.IsEnabled(this, &SMeshProxyDialog::HardAngleThresholdEnabled)
+							.Text(LOCTEXT("HardAngleLabel", "Hard Angle"))
+							.Font(FEditorStyle::GetFontStyle("StandardDialog.SmallFont"))
+						]
+					]
+					+SHorizontalBox::Slot()
+					.FillWidth(0.5f)
+					.HAlign(HAlign_Left)
+					.VAlign(VAlign_Center)
+					[
+						SNew(SBox)
+						.HAlign(HAlign_Fill)
+						.MinDesiredWidth(100.0f)
+						.MaxDesiredWidth(100.0f)
+						[
+							SNew(SNumericEntryBox<float>)
+							.Font(FEditorStyle::GetFontStyle("StandardDialog.SmallFont"))
+							.MinValue(0.f)
+							.MaxValue(180.f)
+							.MinSliderValue(0.f)
+							.MaxSliderValue(180.f)
+							.AllowSpin(true)
+							.Value(this, &SMeshProxyDialog::GetHardAngleThreshold)
+							.OnValueChanged(this, &SMeshProxyDialog::HardAngleThresholdChanged)
+							.IsEnabled(this, &SMeshProxyDialog::HardAngleThresholdEnabled)
+						]
+					]
+				]
+
+				+SVerticalBox::Slot()
+				.AutoHeight()
+				.Padding(FEditorStyle::GetMargin("StandardDialog.ContentPadding"))
+				[
+					SNew(SCheckBox)
+					.Type(ESlateCheckBoxType::CheckBox)
+					.IsChecked(this, &SMeshProxyDialog::GetUseClippingPlane)
+					.OnCheckStateChanged(this, &SMeshProxyDialog::SetUseClippingPlane)
+					.Content()
+					[
+						SNew(STextBlock)
+						.Text(LOCTEXT("ClippingPlaneLabel", "Use Clipping Plane"))
+						.Font(FEditorStyle::GetFontStyle("StandardDialog.SmallFont"))
+					]
+				]
+
+				+SVerticalBox::Slot()
+				.AutoHeight()
+				.Padding(FEditorStyle::GetMargin("StandardDialog.ContentPadding"))
+				[
+					SNew(SHorizontalBox)
+					+SHorizontalBox::Slot()
+					.VAlign(VAlign_Center)
+					.FillWidth(0.5f)
+					[
+						SNew(SHorizontalBox)
+						+SHorizontalBox::Slot()
+						.AutoWidth()
+						[
+							SNew(SBox)
+							.MinDesiredWidth(30)
+							[
+								SNullWidget::NullWidget
+							]
+						]
+						+SHorizontalBox::Slot()
+						.AutoWidth()
+						.HAlign(HAlign_Left)
+						[
+							SNew(STextBlock)
+							.IsEnabled(this, &SMeshProxyDialog::UseClippingPlaneEnabled)
+							.Text(LOCTEXT("ClippingAxisLabel", "Clipping Axis"))
+							.Font(FEditorStyle::GetFontStyle("StandardDialog.SmallFont"))
+						]
+					]
+					+SHorizontalBox::Slot()
+					.FillWidth(0.5f)
+					.HAlign(HAlign_Left)
+					.VAlign(VAlign_Center)
+					.Padding(8.0, 0.0, 8.0, 0.0)
+					[
+						SNew(STextComboBox)
+						.Font(FEditorStyle::GetFontStyle("StandardDialog.SmallFont"))
+						.OptionsSource(&CuttingPlaneOptions)
+						.InitiallySelectedItem(CuttingPlaneOptions[0])
+						.OnSelectionChanged(this, &SMeshProxyDialog::SetClippingAxis)
+						.IsEnabled(this, &SMeshProxyDialog::UseClippingPlaneEnabled)
+					]
+				]
+
+				+SVerticalBox::Slot()
+				.AutoHeight()
+				.Padding(FEditorStyle::GetMargin("StandardDialog.ContentPadding"))
+				[
+					SNew(SHorizontalBox)
+					+SHorizontalBox::Slot()
+					.FillWidth(0.5f)
+					.VAlign(VAlign_Center)
+					.Padding(0.0, 0.0, 3.0, 0.0)
+					[
+						SNew(SHorizontalBox)
+						+SHorizontalBox::Slot()
+						.AutoWidth()
+						[
+							SNew(SBox)
+							.MinDesiredWidth(30)
+							[
+								SNullWidget::NullWidget
+							]
+						]
+						+SHorizontalBox::Slot()
+						.HAlign(HAlign_Left)
+						.AutoWidth()
+						[
+							SNew(STextBlock)
+							.Text(LOCTEXT("PlaneLevelLabel", "Plane level"))
+							.Font(FEditorStyle::GetFontStyle("StandardDialog.SmallFont"))
+							.IsEnabled(this, &SMeshProxyDialog::UseClippingPlaneEnabled)
+						]
+					]
+					+SHorizontalBox::Slot()
+					.FillWidth(0.5f)
+					.HAlign(HAlign_Left)
+					[
+						SNew(SBox)
+						.HAlign(HAlign_Fill)
+						.MinDesiredWidth(100.0f)
+						.MaxDesiredWidth(100.0f)
+						[
+							SNew(SNumericEntryBox<float>)
+							.Font(FEditorStyle::GetFontStyle("StandardDialog.SmallFont"))
+							.Value(this, &SMeshProxyDialog::GetClippingLevel)
+							.OnValueChanged(this, &SMeshProxyDialog::ClippingLevelChanged)
+							.IsEnabled(this, &SMeshProxyDialog::UseClippingPlaneEnabled)
+						]
+					]
+				]
+
+				+SVerticalBox::Slot()
+				.AutoHeight()
+				.Padding(FEditorStyle::GetMargin("StandardDialog.ContentPadding"))
+				[
+					SNew(SCheckBox)
+					.Type(ESlateCheckBoxType::CheckBox)
+					.IsChecked(this, &SMeshProxyDialog::GetExportNormalMap)
+					.OnCheckStateChanged(this, &SMeshProxyDialog::SetExportNormalMap)
+					.Content()
+					[
+						SNew(STextBlock)
+						.Text(LOCTEXT("ExportNormalMapLabel", "Export Normal Map"))
+						.Font(FEditorStyle::GetFontStyle("StandardDialog.SmallFont"))
+					]
+				]
+
+				+SVerticalBox::Slot()
+				.AutoHeight()
+				.Padding(FEditorStyle::GetMargin("StandardDialog.ContentPadding"))
+				[
+					SNew(SCheckBox)
+					.Type(ESlateCheckBoxType::CheckBox)
+					.IsChecked(this, &SMeshProxyDialog::GetExportMetallicMap)
+					.OnCheckStateChanged(this, &SMeshProxyDialog::SetExportMetallicMap)
+					.Content()
+					[
+						SNew(STextBlock)
+						.Text(LOCTEXT("ExportMetallicMapLabel", "Export Metallic Map"))
+						.Font(FEditorStyle::GetFontStyle("StandardDialog.SmallFont"))
+					]
+				]
+
+				+SVerticalBox::Slot()
+				.AutoHeight()
+				.Padding(FEditorStyle::GetMargin("StandardDialog.ContentPadding"))
+				[
+					SNew(SCheckBox)
+					.Type(ESlateCheckBoxType::CheckBox)
+					.IsChecked(this, &SMeshProxyDialog::GetExportRoughnessMap)
+					.OnCheckStateChanged(this, &SMeshProxyDialog::SetExportRoughnessMap)
+					.Content()
+					[
+						SNew(STextBlock)
+						.Text(LOCTEXT("ExportRoughnessMapLabel", "Export Roughness Map"))
+						.Font(FEditorStyle::GetFontStyle("StandardDialog.SmallFont"))
+					]
+				]
+
+				+SVerticalBox::Slot()
+				.AutoHeight()
+				.Padding(FEditorStyle::GetMargin("StandardDialog.ContentPadding"))
+				[
+					SNew(SCheckBox)
+					.Type(ESlateCheckBoxType::CheckBox)
+					.IsChecked(this, &SMeshProxyDialog::GetExportSpecularMap)
+					.OnCheckStateChanged(this, &SMeshProxyDialog::SetExportSpecularMap)
+					.Content()
+					[
+						SNew(STextBlock)
+						.Text(LOCTEXT("ExportSpecularMapLabel", "Export Specular Map"))
+						.Font(FEditorStyle::GetFontStyle("StandardDialog.SmallFont"))
+					]
+				]
+			]
+		]
+	];
+}
+END_SLATE_FUNCTION_BUILD_OPTIMIZATION
+
+//Screen size
+TOptional<int32> SMeshProxyDialog::GetScreenSize() const
+{
+	return Tool->ProxySettings.ScreenSize;
+}
+
+void SMeshProxyDialog::ScreenSizeChanged(int32 NewValue)
+{
+	Tool->ProxySettings.ScreenSize = NewValue;
+}
+
+//Recalculate normals
+ECheckBoxState SMeshProxyDialog::GetRecalculateNormals() const
+{
+	return Tool->ProxySettings.bRecalculateNormals ? ECheckBoxState::Checked: ECheckBoxState::Unchecked;
+}
+
+void SMeshProxyDialog::SetRecalculateNormals(ECheckBoxState NewValue)
+{
+	Tool->ProxySettings.bRecalculateNormals = (NewValue == ECheckBoxState::Checked);
+}
+
+//Hard Angle Threshold
+bool SMeshProxyDialog::HardAngleThresholdEnabled() const
+{
+	if(Tool->ProxySettings.bRecalculateNormals)
+	{
+		return true;
+	}
+
+	return false;
+}
+
+TOptional<float> SMeshProxyDialog::GetHardAngleThreshold() const
+{
+	return Tool->ProxySettings.HardAngleThreshold;
+}
+
+void SMeshProxyDialog::HardAngleThresholdChanged(float NewValue)
+{
+	Tool->ProxySettings.HardAngleThreshold = NewValue;
+}
+
+//Merge Distance
+TOptional<int32> SMeshProxyDialog::GetMergeDistance() const
+{
+	return Tool->ProxySettings.MergeDistance;
+}
+
+void SMeshProxyDialog::MergeDistanceChanged(int32 NewValue)
+{
+	Tool->ProxySettings.MergeDistance = NewValue;
+}
+
+//Clipping Plane
+ECheckBoxState SMeshProxyDialog::GetUseClippingPlane() const
+{
+	return Tool->ProxySettings.bUseClippingPlane ? ECheckBoxState::Checked: ECheckBoxState::Unchecked;
+}
+
+void SMeshProxyDialog::SetUseClippingPlane(ECheckBoxState NewValue)
+{
+	Tool->ProxySettings.bUseClippingPlane = (NewValue == ECheckBoxState::Checked);
+}
+
+bool SMeshProxyDialog::UseClippingPlaneEnabled() const
+{
+	if(Tool->ProxySettings.bUseClippingPlane)
+	{
+		return true;
+	}
+	return false;
+}
+
+void SMeshProxyDialog::SetClippingAxis(TSharedPtr<FString> NewSelection, ESelectInfo::Type SelectInfo )
+{
+	//This is a ugly hack, but it solves the problem for now
+	int32 AxisIndex = CuttingPlaneOptions.Find(NewSelection);
+	//FMessageDialog::Debugf(*FString::Printf(TEXT("%d"), AxisIndex ));
+	if(AxisIndex < 3)
+	{
+		Tool->ProxySettings.bPlaneNegativeHalfspace = false;
+		Tool->ProxySettings.AxisIndex = AxisIndex;
+		return;
+	}
+	else
+	{
+		Tool->ProxySettings.bPlaneNegativeHalfspace = true;
+		Tool->ProxySettings.AxisIndex = AxisIndex - 3;
+		return;
+	}
+}
+
+TOptional<float> SMeshProxyDialog::GetClippingLevel() const
+{
+	return Tool->ProxySettings.ClippingLevel;
+}
+
+void SMeshProxyDialog::ClippingLevelChanged(float NewValue)
+{
+	Tool->ProxySettings.ClippingLevel = NewValue;
+}
+
+//Texture Resolution
+void SMeshProxyDialog::SetTextureResolution(TSharedPtr<FString> NewSelection, ESelectInfo::Type SelectInfo)
+{
+	int32 Resolution = 512;
+	TTypeFromString<int32>::FromString(Resolution, **NewSelection);
+	
+	Tool->ProxySettings.TextureWidth	= Resolution;
+	Tool->ProxySettings.TextureHeight = Resolution;
+}
+
+ECheckBoxState SMeshProxyDialog::GetExportNormalMap() const
+{
+	return Tool->ProxySettings.bExportNormalMap ? ECheckBoxState::Checked :  ECheckBoxState::Unchecked;
+}
+
+void SMeshProxyDialog::SetExportNormalMap(ECheckBoxState NewValue)
+{
+	Tool->ProxySettings.bExportNormalMap = (NewValue == ECheckBoxState::Checked);
+}
+
+ECheckBoxState SMeshProxyDialog::GetExportMetallicMap() const
+{
+	return Tool->ProxySettings.bExportMetallicMap ? ECheckBoxState::Checked :  ECheckBoxState::Unchecked;
+}
+
+void SMeshProxyDialog::SetExportMetallicMap(ECheckBoxState NewValue)
+{
+	Tool->ProxySettings.bExportMetallicMap = (NewValue == ECheckBoxState::Checked);
+}
+
+ECheckBoxState SMeshProxyDialog::GetExportRoughnessMap() const
+{
+	return Tool->ProxySettings.bExportRoughnessMap ? ECheckBoxState::Checked :  ECheckBoxState::Unchecked;
+}
+
+void SMeshProxyDialog::SetExportRoughnessMap(ECheckBoxState NewValue)
+{
+	Tool->ProxySettings.bExportRoughnessMap = (NewValue == ECheckBoxState::Checked);
+}
+
+ECheckBoxState SMeshProxyDialog::GetExportSpecularMap() const
+{
+	return Tool->ProxySettings.bExportSpecularMap ? ECheckBoxState::Checked :  ECheckBoxState::Unchecked;
+}
+
+void SMeshProxyDialog::SetExportSpecularMap(ECheckBoxState NewValue)
+{
+	Tool->ProxySettings.bExportSpecularMap = (NewValue == ECheckBoxState::Checked);
+}
+
+
+#undef LOCTEXT_NAMESPACE
diff --git a/Engine/Source/Editor/MergeActors/Private/MeshProxyTool/SMeshProxyDialog.h b/Engine/Source/Editor/MergeActors/Private/MeshProxyTool/SMeshProxyDialog.h
new file mode 100644
--- /dev/null
+++ b/Engine/Source/Editor/MergeActors/Private/MeshProxyTool/SMeshProxyDialog.h
@@ -0,0 +1,70 @@
+// Copyright 1998-2015 Epic Games, Inc. All Rights Reserved.
+
+#pragma once
+
+
+/*-----------------------------------------------------------------------------
+	SMeshProxyDialog
+-----------------------------------------------------------------------------*/
+class SMeshProxyDialog : public SCompoundWidget
+{
+public:
+	SLATE_BEGIN_ARGS(SMeshProxyDialog)
+	{
+	}
+	SLATE_END_ARGS()
+
+public:
+	/** SWidget functions */
+	void Construct(const FArguments& InArgs, FMeshProxyTool* InTool);
+
+protected:
+	/** ScreenSize accessors */
+	TOptional<int32> GetScreenSize() const;
+	void ScreenSizeChanged(int32 NewValue);		//used with editable text block (Simplygon)
+
+	/** Recalculate Normals accessors */
+	ECheckBoxState GetRecalculateNormals() const;
+	void SetRecalculateNormals(ECheckBoxState NewValue);
+
+	/** Hard Angle Threshold accessors */
+	TOptional<float> GetHardAngleThreshold() const;
+	bool HardAngleThresholdEnabled() const;
+	void HardAngleThresholdChanged(float NewValue);
+
+	/** Hole filling accessors */
+	TOptional<int32> GetMergeDistance() const;
+	void MergeDistanceChanged(int32 NewValue);
+
+	/** Clipping Plane accessors */
+	ECheckBoxState GetUseClippingPlane() const;
+	void SetUseClippingPlane(ECheckBoxState NewValue);
+	bool UseClippingPlaneEnabled() const;
+	void SetClippingAxis(TSharedPtr<FString> NewSelection, ESelectInfo::Type SelectInfo);
+	TOptional<float> GetClippingLevel() const;
+	void ClippingLevelChanged(float NewValue);
+
+	/** TextureResolution accessors */
+	void SetTextureResolution(TSharedPtr<FString> NewSelection, ESelectInfo::Type SelectInfo);
+
+	/** Export material properties acessors **/
+	ECheckBoxState GetExportNormalMap() const;
+	void SetExportNormalMap(ECheckBoxState NewValue);
+	ECheckBoxState GetExportMetallicMap() const;
+	void SetExportMetallicMap(ECheckBoxState NewValue);
+	ECheckBoxState GetExportRoughnessMap() const;
+	void SetExportRoughnessMap(ECheckBoxState NewValue);
+	ECheckBoxState GetExportSpecularMap() const;
+	void SetExportSpecularMap(ECheckBoxState NewValue);
+
+private:
+	/** Creates the geometry mode controls */
+	void CreateLayout();
+
+private:
+	FMeshProxyTool* Tool;
+
+	TArray< TSharedPtr<FString> >	CuttingPlaneOptions;
+	TArray< TSharedPtr<FString> >	TextureResolutionOptions;
+};
+
diff --git a/Engine/Source/Editor/MergeActors/Private/SMergeActorsToolbar.cpp b/Engine/Source/Editor/MergeActors/Private/SMergeActorsToolbar.cpp
new file mode 100644
--- /dev/null
+++ b/Engine/Source/Editor/MergeActors/Private/SMergeActorsToolbar.cpp
@@ -0,0 +1,229 @@
+// Copyright 1998-2014 Epic Games, Inc. All Rights Reserved.
+
+#include "MergeActorsPrivatePCH.h"
+#include "SMergeActorsToolbar.h"
+#include "LevelEditor.h"
+#include "IDocumentation.h"
+#include "ContentBrowserModule.h"
+
+#define LOCTEXT_NAMESPACE "SMergeActorsToolbar"
+
+
+//////////////////////////////////////////////////////////////
+
+void SMergeActorsToolbar::Construct(const FArguments& InArgs)
+{
+	// Important: We use raw bindings here because we are releasing our binding in our destructor (where a weak pointer would be invalid)
+	// It's imperative that our delegate is removed in the destructor for the level editor module to play nicely with reloading.
+
+	FLevelEditorModule& LevelEditor = FModuleManager::GetModuleChecked<FLevelEditorModule>("LevelEditor");
+	LevelEditor.OnActorSelectionChanged().AddRaw(this, &SMergeActorsToolbar::OnActorSelectionChanged);
+
+	RegisteredTools = InArgs._ToolsToRegister;
+
+	ChildSlot
+	[
+		SNew(SVerticalBox)
+
+		+ SVerticalBox::Slot()
+		.AutoHeight()
+		.HAlign(HAlign_Left)
+		.Padding(0, 0, 0, 0)
+		[
+			SAssignNew(ToolbarContainer, SBorder)
+			.BorderImage(FEditorStyle::GetBrush("NoBorder"))
+			.Padding(FMargin(4, 0, 0, 0))
+		]
+
+		+ SVerticalBox::Slot()
+		.FillHeight(1.0f)
+		.Padding(2, 0, 0, 0)
+		[
+			SNew(SBorder)
+			.BorderImage(FEditorStyle::GetBrush("ToolPanel.GroupBorder"))
+			.Padding(0.f)
+			.IsEnabled(this, &SMergeActorsToolbar::GetContentEnabledState)
+			[
+				SNew(SVerticalBox)
+
+				+ SVerticalBox::Slot()
+				.FillHeight(1.0f)
+				.Padding(4, 4, 4, 4)
+				[
+					SNew(SScrollBox)
+					+SScrollBox::Slot()
+					[
+						SAssignNew(InlineContentHolder, SBox)
+					]
+				]
+
+				+ SVerticalBox::Slot()
+				.AutoHeight()
+				.HAlign(HAlign_Right)
+				.Padding(4, 4, 10, 4)
+				[
+					SNew(SButton)
+					.Text(LOCTEXT("MergeActors", "Merge Actors"))
+					.OnClicked(this, &SMergeActorsToolbar::OnMergeActorsClicked)
+				]
+			]
+		]
+	];
+
+	UpdateToolbar();
+}
+
+
+SMergeActorsToolbar::~SMergeActorsToolbar()
+{
+	FLevelEditorModule& LevelEditor = FModuleManager::GetModuleChecked<FLevelEditorModule>("LevelEditor");
+	LevelEditor.OnActorSelectionChanged().RemoveAll(this);
+}
+
+
+void SMergeActorsToolbar::OnActorSelectionChanged(const TArray<UObject*>& NewSelection, bool bForceRefresh)
+{
+	SelectedObjects = NewSelection;
+}
+
+
+void SMergeActorsToolbar::OnToolSelectionChanged(const ECheckBoxState NewCheckedState, int32 ToolIndex)
+{
+	if (NewCheckedState == ECheckBoxState::Checked)
+	{
+		CurrentlySelectedTool = ToolIndex;
+		UpdateInlineContent();
+	}
+}
+
+
+ECheckBoxState SMergeActorsToolbar::OnIsToolSelected(int32 ToolIndex) const
+{
+	return (CurrentlySelectedTool == ToolIndex) ? ECheckBoxState::Checked : ECheckBoxState::Unchecked;
+}
+
+
+FReply SMergeActorsToolbar::OnMergeActorsClicked()
+{
+	if (CurrentlySelectedTool >= 0 && CurrentlySelectedTool < RegisteredTools.Num())
+	{
+		const FString DefaultPackageName = RegisteredTools[CurrentlySelectedTool]->GetDefaultPackageName();
+		const FString DefaultPath = FPackageName::GetLongPackagePath(DefaultPackageName);
+		const FString DefaultName = FPackageName::GetShortName(DefaultPackageName);
+
+		// Initialize SaveAssetDialog config
+		FSaveAssetDialogConfig SaveAssetDialogConfig;
+		SaveAssetDialogConfig.DialogTitleOverride = LOCTEXT("CreateMergedActorTitle", "Create Merged Actor");
+		SaveAssetDialogConfig.DefaultPath = DefaultPath;
+		SaveAssetDialogConfig.DefaultAssetName = DefaultName;
+		SaveAssetDialogConfig.ExistingAssetPolicy = ESaveAssetDialogExistingAssetPolicy::AllowButWarn;
+
+		FContentBrowserModule& ContentBrowserModule = FModuleManager::LoadModuleChecked<FContentBrowserModule>("ContentBrowser");
+		FString SaveObjectPath = ContentBrowserModule.Get().CreateModalSaveAssetDialog(SaveAssetDialogConfig);
+		if (!SaveObjectPath.IsEmpty())
+		{
+			const FString PackageName = FPackageName::ObjectPathToPackageName(SaveObjectPath);
+
+			RegisteredTools[CurrentlySelectedTool]->RunMerge(PackageName);
+		}
+	}
+
+	return FReply::Handled();
+}
+
+
+bool SMergeActorsToolbar::GetContentEnabledState() const
+{
+	int32 NumSelectedStaticMeshes = 0;
+	for (UObject* Object : SelectedObjects)
+	{
+		if (AActor* Actor = Cast<AActor>(Object))
+		{
+			TInlineComponentArray<UStaticMeshComponent*> StaticMeshComponents;
+			Actor->GetComponents(StaticMeshComponents);
+			NumSelectedStaticMeshes += StaticMeshComponents.Num();
+			if (NumSelectedStaticMeshes >= 2)
+			{
+				break;
+			}
+		}
+	}
+	return NumSelectedStaticMeshes >= 2;
+}
+
+
+void SMergeActorsToolbar::AddTool(IMergeActorsTool* Tool)
+{
+	check(!RegisteredTools.Contains(Tool));
+	RegisteredTools.Add(Tool);
+	UpdateToolbar();
+}
+
+
+void SMergeActorsToolbar::RemoveTool(IMergeActorsTool* Tool)
+{
+	int32 IndexToRemove = RegisteredTools.Find(Tool);
+	if (IndexToRemove != INDEX_NONE)
+	{
+		RegisteredTools.RemoveAt(IndexToRemove);
+
+		if (CurrentlySelectedTool > IndexToRemove)
+		{
+			CurrentlySelectedTool--;
+		}
+		UpdateToolbar();
+	}
+}
+
+
+void SMergeActorsToolbar::UpdateToolbar()
+{
+	const ISlateStyle& StyleSet = FEditorStyle::Get();
+
+	TSharedRef<SHorizontalBox> HorizontalBox =
+		SNew(SHorizontalBox);
+
+	for (int32 ToolIndex = 0; ToolIndex < RegisteredTools.Num(); ToolIndex++)
+	{
+		const IMergeActorsTool* Tool = RegisteredTools[ToolIndex];
+
+		HorizontalBox->AddSlot()
+		.Padding(StyleSet.GetMargin("EditorModesToolbar.SToolBarButtonBlock.Padding"))
+		[
+			SNew(SCheckBox)
+			.Style(&StyleSet, "EditorModesToolbar.ToggleButton")
+			.OnCheckStateChanged(this, &SMergeActorsToolbar::OnToolSelectionChanged, ToolIndex)
+			.IsChecked(this, &SMergeActorsToolbar::OnIsToolSelected, ToolIndex)
+			.Padding(StyleSet.GetMargin("EditorModesToolbar.SToolBarButtonBlock.CheckBox.Padding"))
+			.ToolTip(IDocumentation::Get()->CreateToolTip(Tool->GetTooltipText(), nullptr, FString(), FString()))
+			[
+				SNew(SImage)
+				.Image(StyleSet.GetBrush(Tool->GetIconName()))
+			]
+		];
+	}
+
+	TSharedRef<SBorder> ToolbarContent =
+		SNew(SBorder)
+		.Padding(0)
+		.BorderImage(StyleSet.GetBrush("NoBorder"))
+		[
+			HorizontalBox
+		];
+
+	ToolbarContainer->SetContent(ToolbarContent);
+
+	UpdateInlineContent();
+}
+
+
+void SMergeActorsToolbar::UpdateInlineContent()
+{
+	if (CurrentlySelectedTool >= 0 && CurrentlySelectedTool < RegisteredTools.Num())
+	{
+		InlineContentHolder->SetContent(RegisteredTools[CurrentlySelectedTool]->GetWidget());
+	}
+}
+
+
+#undef LOCTEXT_NAMESPACE
diff --git a/Engine/Source/Editor/MergeActors/Private/SMergeActorsToolbar.h b/Engine/Source/Editor/MergeActors/Private/SMergeActorsToolbar.h
new file mode 100644
--- /dev/null
+++ b/Engine/Source/Editor/MergeActors/Private/SMergeActorsToolbar.h
@@ -0,0 +1,80 @@
+// Copyright 1998-2014 Epic Games, Inc. All Rights Reserved.
+
+#pragma once
+
+class IMergeActorsTool;
+
+//////////////////////////////////////////////////////////////////////////
+// SMergeActorsToolbar
+
+class SMergeActorsToolbar : public SCompoundWidget
+{
+public:
+
+	SLATE_BEGIN_ARGS(SMergeActorsToolbar) {}
+		SLATE_ARGUMENT(TArray<IMergeActorsTool*>, ToolsToRegister)
+	SLATE_END_ARGS()
+
+	/**
+	 * Construct the widget
+	 *
+	 * @param	InArgs			A declaration from which to construct the widget
+	 */
+	void Construct(const FArguments& InArgs);
+
+	/** Constructor */
+	SMergeActorsToolbar()
+		: CurrentlySelectedTool(0)
+	{}
+
+	/** Destructor */
+	virtual ~SMergeActorsToolbar();
+
+	/** Add a new tool to the toolbar */
+	void AddTool(IMergeActorsTool* Tool);
+
+	/** Remove an existing tool from the toolbar */
+	void RemoveTool(IMergeActorsTool* Tool);
+
+
+private:
+
+	/** Called when the level actor selection changes */
+	void OnActorSelectionChanged(const TArray<UObject*>& NewSelection, bool bForceRefresh);
+
+	/** Called when the currently selected tool changes */
+	void OnToolSelectionChanged(const ECheckBoxState NewCheckedState, int32 ToolIndex);
+
+	/** Called to determine whether a toolbox widget is selected or not */
+	ECheckBoxState OnIsToolSelected(int32 ToolIndex) const;
+
+	/** Called when the Merge Actors button is clicked */
+	FReply OnMergeActorsClicked();
+
+	/** Called to determine whether the widget content is enabled or not */
+	bool GetContentEnabledState() const;
+
+	/** Update the toolbar container based on the currently registered tools */
+	void UpdateToolbar();
+
+	/** Updates the inline content widget for the current tool */
+	void UpdateInlineContent();
+
+
+private:
+
+	/** List of registered tool instances */
+	TArray<IMergeActorsTool*> RegisteredTools;
+
+	/** Index of currently selected tool */
+	int32 CurrentlySelectedTool;
+
+	/** List of currently selected objects */
+	TArray<UObject*> SelectedObjects;
+
+	/** The container holding the toolbar */
+	TSharedPtr<SBorder> ToolbarContainer;
+
+	/** Inline content area for different tool modes */
+	TSharedPtr<SBox> InlineContentHolder;
+};
diff --git a/Engine/Source/Editor/MergeActors/Public/IMergeActorsModule.h b/Engine/Source/Editor/MergeActors/Public/IMergeActorsModule.h
new file mode 100644
--- /dev/null
+++ b/Engine/Source/Editor/MergeActors/Public/IMergeActorsModule.h
@@ -0,0 +1,33 @@
+// Copyright 1998-2014 Epic Games, Inc. All Rights Reserved.
+#pragma once
+
+#include "ModuleManager.h"
+#include "IMergeActorsTool.h"
+
+
+/**
+ * Merge Actors module interface
+ */
+class IMergeActorsModule : public IModuleInterface
+{
+
+public:
+
+	/**
+	 * Get reference to the Merge Actors module instance
+	 */
+	static inline IMergeActorsModule& Get()
+	{
+		return FModuleManager::LoadModuleChecked<IMergeActorsModule>("MergeActors");
+	}
+
+	/**
+	 * Register an IMergeActorsTool with the module, passing ownership to it
+	 */
+	virtual bool RegisterMergeActorsTool(TUniquePtr<IMergeActorsTool> Tool) = 0;
+
+	/**
+	 * Unregister an IMergeActorsTool with the module
+	 */
+	virtual bool UnregisterMergeActorsTool(IMergeActorsTool* Tool) = 0;
+};
diff --git a/Engine/Source/Editor/MergeActors/Public/IMergeActorsTool.h b/Engine/Source/Editor/MergeActors/Public/IMergeActorsTool.h
new file mode 100644
--- /dev/null
+++ b/Engine/Source/Editor/MergeActors/Public/IMergeActorsTool.h
@@ -0,0 +1,43 @@
+// Copyright 1998-2014 Epic Games, Inc. All Rights Reserved.
+#pragma once
+
+/**
+* Merge Actors tool interface
+*/
+class IMergeActorsTool
+{
+
+public:
+
+	/** Virtual destructor */
+	virtual ~IMergeActorsTool() {}
+
+	/**
+	 * Gets the widget instance associated with this tool
+	 */
+	virtual TSharedRef<SWidget> GetWidget() = 0;
+
+	/**
+	 * Get the name of the icon displayed in the Merge Actors toolbar
+	 */
+	virtual FName GetIconName() const = 0;
+
+	/**
+	 * Get Tooltip text displayed in the Merge Actors toolbar
+	 */
+	virtual FText GetTooltipText() const = 0;
+
+	/**
+	 * Get default name for the merged asset package
+	 */
+	virtual FString GetDefaultPackageName() const = 0;
+
+	/**
+	 * Perform merge operation
+	 *
+	 * @param	PackageName		Long package name of the asset to create
+	 * @return	true if the merge succeeded
+	 */
+	virtual bool RunMerge(const FString& PackageName) = 0;
+
+};
diff --git a/Engine/Source/Editor/UnrealEd/Classes/Editor/UnrealEdEngine.h b/Engine/Source/Editor/UnrealEd/Classes/Editor/UnrealEdEngine.h
--- a/Engine/Source/Editor/UnrealEd/Classes/Editor/UnrealEdEngine.h
+++ b/Engine/Source/Editor/UnrealEd/Classes/Editor/UnrealEdEngine.h
@@ -424,16 +424,6 @@ public:
 	 */
 	virtual void edactRemoveFromGroup();
 	
-	/**
-	 * Opens the dialog window for merging selected actors into single static mesh
-	 */
-	virtual void edactMergeActors();
-
-	/**
-	 * Merges selected actors geometry grouping them by materials 
-	 */
-	virtual void edactMergeActorsByMaterials();
-
 	/**
 	 * Copy selected actors to the clipboard.  Does not copy PrefabInstance actors or parts of Prefabs.
 	 *
diff --git a/Engine/Source/Editor/UnrealEd/Classes/Settings/EditorExperimentalSettings.h b/Engine/Source/Editor/UnrealEd/Classes/Settings/EditorExperimentalSettings.h
--- a/Engine/Source/Editor/UnrealEd/Classes/Settings/EditorExperimentalSettings.h
+++ b/Engine/Source/Editor/UnrealEd/Classes/Settings/EditorExperimentalSettings.h
@@ -49,10 +49,6 @@ public:
 	UPROPERTY(EditAnywhere, config, Category=Tools, meta=(DisplayName="Messaging Debugger"))
 	bool bMessagingDebugger;
 
-	/** Allows to use actor merging utilities (Simplygon Proxy LOD, Grouping by Materials)*/
-	UPROPERTY(EditAnywhere, config, Category=Tools, meta=(DisplayName="Actor Merging"))
-	bool bActorMerging;
-
 	/** Specify which console-specific nomenclature to use for gamepad label text */
 	UPROPERTY(EditAnywhere, config, Category=UserInterface, meta=(DisplayName="Console for Gamepad Labels"))
 	TEnumAsByte<EConsoleForGamepadLabels::Type> ConsoleForGamepadLabels;
diff --git a/Engine/Source/Editor/UnrealEd/Private/Dialogs/SMeshMergingDialog.cpp b/Engine/Source/Editor/UnrealEd/Private/Dialogs/SMeshMergingDialog.cpp
deleted file mode 100644
diff --git a/Engine/Source/Editor/UnrealEd/Private/Dialogs/SMeshProxyDialog.cpp b/Engine/Source/Editor/UnrealEd/Private/Dialogs/SMeshProxyDialog.cpp
deleted file mode 100644
diff --git a/Engine/Source/Editor/UnrealEd/Private/EditorEngine.cpp b/Engine/Source/Editor/UnrealEd/Private/EditorEngine.cpp
--- a/Engine/Source/Editor/UnrealEd/Private/EditorEngine.cpp
+++ b/Engine/Source/Editor/UnrealEd/Private/EditorEngine.cpp
@@ -643,7 +643,8 @@ void UEditorEngine::Init(IEngineLoop* InEngineLoop)
 			TEXT("LocalizationDashboard"),
 			TEXT("ReferenceViewer"),
 			TEXT("TreeMap"),
-			TEXT("SizeMap")
+			TEXT("SizeMap"),
+			TEXT("MergeActors")
 		};
 
 		FScopedSlowTask ModuleSlowTask(ARRAY_COUNT(ModuleNames));
diff --git a/Engine/Source/Editor/UnrealEd/Private/EditorGroup.cpp b/Engine/Source/Editor/UnrealEd/Private/EditorGroup.cpp
--- a/Engine/Source/Editor/UnrealEd/Private/EditorGroup.cpp
+++ b/Engine/Source/Editor/UnrealEd/Private/EditorGroup.cpp
@@ -4,9 +4,6 @@
 #include "UnrealEd.h"
 #include "ScopedTransaction.h"
 #include "MainFrame.h"
-#include "Dialogs/SMeshProxyDialog.h"
-#include "Dialogs/SMeshMergingDialog.h"
-#include "Dialogs/DlgPickAssetPath.h"
 #include "MeshUtilities.h"
 #include "AssetRegistryModule.h"
 #include "ContentBrowserModule.h"
@@ -207,73 +204,3 @@ void UUnrealEdEngine::edactRemoveFromGroup()
 	}
 }
 
-void UUnrealEdEngine::edactMergeActors()
-{
-	IMeshUtilities& MeshUtilities = FModuleManager::Get().LoadModuleChecked<IMeshUtilities>("MeshUtilities");
-	if (MeshUtilities.GetMeshMergingInterface() != nullptr)
-	{
-		TSharedPtr<IMeshProxyDialog> MeshProxyControls;
-		if(MeshProxyControls.IsValid() && MeshProxyControls->GetParentWindow().IsValid())
-		{
-			MeshProxyControls->GetParentWindow()->HideWindow();
-		}
-
-		if(!MeshProxyControls.IsValid())
-		{
-			// Create the controls if they haven't already been created
-			MeshProxyControls = IMeshProxyDialog::MakeControls();
-		}
-
-		if(!MeshProxyControls->GetParentWindow().IsValid())
-		{
-			// Create window
-			TSharedRef<SWindow> MeshProxyDialog = SNew(SWindow);
-			MeshProxyDialog->SetTitle(NSLOCTEXT("SMeshProxyDialog", "MeshProxyDialogTitle", "Mesh Proxy"));
-			MeshProxyDialog->Resize(FVector2D(400.0f, 350.0f));
-
-			MeshProxyControls->AssignWindowContent(MeshProxyDialog);
-			MeshProxyControls->SetParentWindow(MeshProxyDialog);
-
-			MeshProxyDialog->SetOnWindowClosed(FOnWindowClosed::CreateSP(MeshProxyControls.ToSharedRef(), &IMeshProxyDialog::OnWindowClosed));
-
-			// Spawn the dialog window
-			TSharedPtr<SWindow> ParentWindow;
-			if(FModuleManager::Get().IsModuleLoaded("MainFrame"))
-			{
-				IMainFrameModule& MainFrame = FModuleManager::GetModuleChecked<IMainFrameModule>("MainFrame");
-				ParentWindow = MainFrame.GetParentWindow();
-			}
-
-			if(ParentWindow.IsValid())
-			{
-				// Parent the window to the main frame 
-				FSlateApplication::Get().AddWindowAsNativeChild(MeshProxyDialog, ParentWindow.ToSharedRef());
-			}
-			else
-			{
-				// Spawn new window
-				FSlateApplication::Get().AddWindow(MeshProxyDialog);
-			}
-		}
-		
-		MeshProxyControls->GetParentWindow()->ShowWindow();
-		MeshProxyControls->MarkDirty();
-	}
-}
-
-void UUnrealEdEngine::edactMergeActorsByMaterials()
-{
-	/** Create the window to host mesh megin widget */
-	TSharedRef<SWindow> MergeActorsWindow = SNew(SWindow)
-											.Title(FText::FromString("Merge actors"))
-											.ClientSize(FVector2D(400, 350));
-
-	/** Set the content of the window to our package dialog widget */
-	TSharedRef<SMeshMergingDialog> MergingDialog = SNew(SMeshMergingDialog)
-													.ParentWindow(MergeActorsWindow);
-
-	MergeActorsWindow->SetContent(MergingDialog);
-
-	/** Show the package dialog window as a modal window */
-	GEditor->EditorAddModalWindow(MergeActorsWindow);
-}
\ No newline at end of file
diff --git a/Engine/Source/Editor/UnrealEd/Public/Dialogs/SMeshMergingDialog.h b/Engine/Source/Editor/UnrealEd/Public/Dialogs/SMeshMergingDialog.h
deleted file mode 100644
diff --git a/Engine/Source/Editor/UnrealEd/Public/Dialogs/SMeshProxyDialog.h b/Engine/Source/Editor/UnrealEd/Public/Dialogs/SMeshProxyDialog.h
deleted file mode 100644
diff --git a/Engine/Source/Editor/UnrealEd/UnrealEd.Build.cs b/Engine/Source/Editor/UnrealEd/UnrealEd.Build.cs
--- a/Engine/Source/Editor/UnrealEd/UnrealEd.Build.cs
+++ b/Engine/Source/Editor/UnrealEd/UnrealEd.Build.cs
@@ -39,6 +39,7 @@ public class UnrealEd : ModuleRules
 				"ImageWrapper",
 				"MainFrame",
 				"MaterialEditor",
+                "MergeActors",
 				"MeshUtilities",
 				"Messaging",
 				"NiagaraEditor",
@@ -183,6 +184,7 @@ public class UnrealEd : ModuleRules
 				"PlacementMode",
 				"NiagaraEditor",
 				"MeshUtilities",
+                "MergeActors",
 				"GameProjectGeneration",
 				"ProjectSettingsViewer",
 				"ProjectTargetPlatformEditor",
