A bunch of non-unity compilation fixes #UE4

diff --git a/Engine/Source/Editor/MergeActors/Private/MeshMergingTool/SMeshMergingDialog.cpp b/Engine/Source/Editor/MergeActors/Private/MeshMergingTool/SMeshMergingDialog.cpp
--- a/Engine/Source/Editor/MergeActors/Private/MeshMergingTool/SMeshMergingDialog.cpp
+++ b/Engine/Source/Editor/MergeActors/Private/MeshMergingTool/SMeshMergingDialog.cpp
@@ -5,6 +5,7 @@
 #include "Dialogs/DlgPickAssetPath.h"
 #include "STextComboBox.h"
 #include "RawMesh.h"
+#include "MeshMergingTool/MeshMergingTool.h"
 
 #define LOCTEXT_NAMESPACE "SMeshMergingDialog"
 
diff --git a/Engine/Source/Editor/MergeActors/Private/MeshProxyTool/SMeshProxyDialog.cpp b/Engine/Source/Editor/MergeActors/Private/MeshProxyTool/SMeshProxyDialog.cpp
--- a/Engine/Source/Editor/MergeActors/Private/MeshProxyTool/SMeshProxyDialog.cpp
+++ b/Engine/Source/Editor/MergeActors/Private/MeshProxyTool/SMeshProxyDialog.cpp
@@ -5,6 +5,7 @@
 #include "Dialogs/DlgPickAssetPath.h"
 #include "SNumericEntryBox.h"
 #include "STextComboBox.h"
+#include "MeshProxyTool/MeshProxyTool.h"
 
 #define LOCTEXT_NAMESPACE "SMeshProxyDialog"
 
diff --git a/Engine/Source/Editor/MergeActors/Private/MeshProxyTool/SMeshProxyDialog.h b/Engine/Source/Editor/MergeActors/Private/MeshProxyTool/SMeshProxyDialog.h
--- a/Engine/Source/Editor/MergeActors/Private/MeshProxyTool/SMeshProxyDialog.h
+++ b/Engine/Source/Editor/MergeActors/Private/MeshProxyTool/SMeshProxyDialog.h
@@ -2,6 +2,7 @@
 
 #pragma once
 
+class FMeshProxyTool;
 
 /*-----------------------------------------------------------------------------
 	SMeshProxyDialog
