Moved MergeActors tool back into Experimental category for now.  Fixed a small bug regarding whether the panel is enabled, and optimized a little.

diff --git a/Engine/Source/Editor/MergeActors/Private/MergeActorsModule.cpp b/Engine/Source/Editor/MergeActors/Private/MergeActorsModule.cpp
--- a/Engine/Source/Editor/MergeActors/Private/MergeActorsModule.cpp
+++ b/Engine/Source/Editor/MergeActors/Private/MergeActorsModule.cpp
@@ -53,6 +53,9 @@ private:
 
 	/** List of registered MergeActorsTool instances */
 	TArray<TUniquePtr<IMergeActorsTool>> MergeActorsTools;
+
+	/** Whether a nomad tab spawner was registered */
+	bool bRegisteredTabSpawner;
 };
 
 IMPLEMENT_MODULE(FMergeActorsModule, MergeActors);
@@ -90,12 +93,21 @@ TSharedRef<SDockTab> FMergeActorsModule::CreateMergeActorsTab(const FSpawnTabArg
 
 void FMergeActorsModule::StartupModule()
 {
+	bRegisteredTabSpawner = false;
+
+	if (!GetDefault<UEditorExperimentalSettings>()->bActorMerging)
+	{
+		return;
+	}
+
 	FGlobalTabmanager::Get()->RegisterNomadTabSpawner(MergeActorsApp, FOnSpawnTab::CreateRaw(this, &FMergeActorsModule::CreateMergeActorsTab))
 		.SetDisplayName(LOCTEXT("TabTitle", "Merge Actors"))
 		.SetTooltipText(LOCTEXT("TooltipText", "Open the Merge Actors tab."))
 		.SetGroup(WorkspaceMenu::GetMenuStructure().GetLevelEditorCategory())
 		.SetIcon(FSlateIcon(FEditorStyle::GetStyleSetName(), "ClassViewer.TabIcon"));
 
+	bRegisteredTabSpawner = true;
+
 	// Register built-in merging tools straight away
 	ensure(RegisterMergeActorsTool(MakeUnique<FMeshMergingTool>()));
 	ensure(RegisterMergeActorsTool(MakeUnique<FMeshProxyTool>()));
@@ -104,6 +116,11 @@ void FMergeActorsModule::StartupModule()
 
 void FMergeActorsModule::ShutdownModule()
 {
+	if (!bRegisteredTabSpawner)
+	{
+		return;
+	}
+
 	if (FSlateApplication::IsInitialized())
 	{
 		FGlobalTabmanager::Get()->UnregisterNomadTabSpawner(MergeActorsApp);
diff --git a/Engine/Source/Editor/MergeActors/Private/SMergeActorsToolbar.cpp b/Engine/Source/Editor/MergeActors/Private/SMergeActorsToolbar.cpp
--- a/Engine/Source/Editor/MergeActors/Private/SMergeActorsToolbar.cpp
+++ b/Engine/Source/Editor/MergeActors/Private/SMergeActorsToolbar.cpp
@@ -71,6 +71,9 @@ void SMergeActorsToolbar::Construct(const FArguments& InArgs)
 	];
 
 	UpdateToolbar();
+
+	// Update selected actor state for the first time
+	GUnrealEd->UpdateFloatingPropertyWindows();
 }
 
 
@@ -84,6 +87,23 @@ SMergeActorsToolbar::~SMergeActorsToolbar()
 void SMergeActorsToolbar::OnActorSelectionChanged(const TArray<UObject*>& NewSelection, bool bForceRefresh)
 {
 	SelectedObjects = NewSelection;
+
+	int32 NumSelectedStaticMeshes = 0;
+	for (UObject* Object : SelectedObjects)
+	{
+		if (AActor* Actor = Cast<AActor>(Object))
+		{
+			TInlineComponentArray<UStaticMeshComponent*> StaticMeshComponents;
+			Actor->GetComponents(StaticMeshComponents);
+			NumSelectedStaticMeshes += StaticMeshComponents.Num();
+			if (NumSelectedStaticMeshes >= 2)
+			{
+				break;
+			}
+		}
+	}
+
+	bIsContentEnabled = (NumSelectedStaticMeshes >= 2);
 }
 
 
@@ -134,21 +154,7 @@ FReply SMergeActorsToolbar::OnMergeActorsClicked()
 
 bool SMergeActorsToolbar::GetContentEnabledState() const
 {
-	int32 NumSelectedStaticMeshes = 0;
-	for (UObject* Object : SelectedObjects)
-	{
-		if (AActor* Actor = Cast<AActor>(Object))
-		{
-			TInlineComponentArray<UStaticMeshComponent*> StaticMeshComponents;
-			Actor->GetComponents(StaticMeshComponents);
-			NumSelectedStaticMeshes += StaticMeshComponents.Num();
-			if (NumSelectedStaticMeshes >= 2)
-			{
-				break;
-			}
-		}
-	}
-	return NumSelectedStaticMeshes >= 2;
+	return bIsContentEnabled;
 }
 
 
diff --git a/Engine/Source/Editor/MergeActors/Private/SMergeActorsToolbar.h b/Engine/Source/Editor/MergeActors/Private/SMergeActorsToolbar.h
--- a/Engine/Source/Editor/MergeActors/Private/SMergeActorsToolbar.h
+++ b/Engine/Source/Editor/MergeActors/Private/SMergeActorsToolbar.h
@@ -51,7 +51,7 @@ private:
 	/** Called when the Merge Actors button is clicked */
 	FReply OnMergeActorsClicked();
 
-	/** Called to determine whether the widget content is enabled or not */
+	/** Determine whether the widget content is enabled or not */
 	bool GetContentEnabledState() const;
 
 	/** Update the toolbar container based on the currently registered tools */
@@ -72,6 +72,9 @@ private:
 	/** List of currently selected objects */
 	TArray<UObject*> SelectedObjects;
 
+	/** Whether the merge actors tool panel is enabled or not */
+	bool bIsContentEnabled;
+
 	/** The container holding the toolbar */
 	TSharedPtr<SBorder> ToolbarContainer;
 
diff --git a/Engine/Source/Editor/MergeActors/Public/IMergeActorsModule.h b/Engine/Source/Editor/MergeActors/Public/IMergeActorsModule.h
--- a/Engine/Source/Editor/MergeActors/Public/IMergeActorsModule.h
+++ b/Engine/Source/Editor/MergeActors/Public/IMergeActorsModule.h
@@ -24,10 +24,10 @@ public:
 	/**
 	 * Register an IMergeActorsTool with the module, passing ownership to it
 	 */
-	virtual bool RegisterMergeActorsTool(TUniquePtr<IMergeActorsTool> Tool) = 0;
+	MERGEACTORS_API virtual bool RegisterMergeActorsTool(TUniquePtr<IMergeActorsTool> Tool) = 0;
 
 	/**
 	 * Unregister an IMergeActorsTool with the module
 	 */
-	virtual bool UnregisterMergeActorsTool(IMergeActorsTool* Tool) = 0;
+	MERGEACTORS_API virtual bool UnregisterMergeActorsTool(IMergeActorsTool* Tool) = 0;
 };
diff --git a/Engine/Source/Editor/UnrealEd/Classes/Settings/EditorExperimentalSettings.h b/Engine/Source/Editor/UnrealEd/Classes/Settings/EditorExperimentalSettings.h
--- a/Engine/Source/Editor/UnrealEd/Classes/Settings/EditorExperimentalSettings.h
+++ b/Engine/Source/Editor/UnrealEd/Classes/Settings/EditorExperimentalSettings.h
@@ -49,6 +49,10 @@ public:
 	UPROPERTY(EditAnywhere, config, Category=Tools, meta=(DisplayName="Messaging Debugger"))
 	bool bMessagingDebugger;
 
+	/** Allows to use actor merging utilities (Simplygon Proxy LOD, Grouping by Materials)*/
+	UPROPERTY(EditAnywhere, config, Category = Tools, meta = (DisplayName = "Actor Merging"))
+	bool bActorMerging;
+
 	/** Specify which console-specific nomenclature to use for gamepad label text */
 	UPROPERTY(EditAnywhere, config, Category=UserInterface, meta=(DisplayName="Console for Gamepad Labels"))
 	TEnumAsByte<EConsoleForGamepadLabels::Type> ConsoleForGamepadLabels;
