Changed behavior so that selecting just a single actor is sufficient to enable the Merge Actors panel.

diff --git a/Engine/Source/Editor/MergeActors/Private/SMergeActorsToolbar.cpp b/Engine/Source/Editor/MergeActors/Private/SMergeActorsToolbar.cpp
--- a/Engine/Source/Editor/MergeActors/Private/SMergeActorsToolbar.cpp
+++ b/Engine/Source/Editor/MergeActors/Private/SMergeActorsToolbar.cpp
@@ -87,23 +87,7 @@ SMergeActorsToolbar::~SMergeActorsToolbar()
 void SMergeActorsToolbar::OnActorSelectionChanged(const TArray<UObject*>& NewSelection, bool bForceRefresh)
 {
 	SelectedObjects = NewSelection;
-
-	int32 NumSelectedStaticMeshes = 0;
-	for (UObject* Object : SelectedObjects)
-	{
-		if (AActor* Actor = Cast<AActor>(Object))
-		{
-			TInlineComponentArray<UStaticMeshComponent*> StaticMeshComponents;
-			Actor->GetComponents(StaticMeshComponents);
-			NumSelectedStaticMeshes += StaticMeshComponents.Num();
-			if (NumSelectedStaticMeshes >= 2)
-			{
-				break;
-			}
-		}
-	}
-
-	bIsContentEnabled = (NumSelectedStaticMeshes >= 2);
+	bIsContentEnabled = (NewSelection.Num() > 0);
 }
 
 
