Merge Actors tool now available without requiring editor restart when enabled in Experimental Settings.  Simplygon tool now only available if Simplygon available.

diff --git a/Engine/Source/Editor/MainFrame/Private/Menus/MainMenu.cpp b/Engine/Source/Editor/MainFrame/Private/Menus/MainMenu.cpp
--- a/Engine/Source/Editor/MainFrame/Private/Menus/MainMenu.cpp
+++ b/Engine/Source/Editor/MainFrame/Private/Menus/MainMenu.cpp
@@ -196,9 +196,10 @@ void FMainMenu::FillWindowMenu( FMenuBuilder& MenuBuilder, const TSharedRef< FEx
 		bool bBlutility = GetDefault<UEditorExperimentalSettings>()->bEnableEditorUtilityBlueprints;
 		bool bLocalizationDashboard = GetDefault<UEditorExperimentalSettings>()->bEnableLocalizationDashboard;
 		bool bTranslationEditor = GetDefault<UEditorExperimentalSettings>()->bEnableTranslationEditor;
+		bool bMergeActors = GetDefault<UEditorExperimentalSettings>()->bActorMerging;
 
 		// Make sure at least one is enabled before creating the section
-		if (bMessagingDebugger || bBlutility || bLocalizationDashboard || bTranslationEditor)
+		if (bMessagingDebugger || bBlutility || bLocalizationDashboard || bTranslationEditor || bMergeActors)
 		{
 			MenuBuilder.BeginSection("ExperimentalTabSpawners", LOCTEXT("ExperimentalTabSpawnersHeading", "Experimental"));
 			{
@@ -250,6 +251,17 @@ void FMainMenu::FillWindowMenu( FMenuBuilder& MenuBuilder, const TSharedRef< FEx
 						FUIAction(FExecuteAction::CreateStatic(&FMainFrameTranslationEditorMenu::HandleOpenTranslationPicker))
 						);
 				}
+
+				// Actor merging
+				if (bMergeActors)
+				{
+					MenuBuilder.AddMenuEntry(
+						LOCTEXT("MergeActorsMenuLabel", "Merge Actors"),
+						LOCTEXT("MergeActorsToolTip", "The Merge Actors tab provides tools for merging multiple actor meshes into a single mesh."),
+						FSlateIcon(),
+						FUIAction(FExecuteAction::CreateStatic(&FMainMenu::OpenMergeActors))
+						);
+				}
 			}
 			MenuBuilder.EndSection();
 		}
diff --git a/Engine/Source/Editor/MainFrame/Private/Menus/MainMenu.h b/Engine/Source/Editor/MainFrame/Private/Menus/MainMenu.h
--- a/Engine/Source/Editor/MainFrame/Private/Menus/MainMenu.h
+++ b/Engine/Source/Editor/MainFrame/Private/Menus/MainMenu.h
@@ -111,4 +111,13 @@ private:
 		FModuleManager::Get().LoadModuleChecked<IModuleInterface>("LogVisualizer");
 		FGlobalTabmanager::Get()->InvokeTab(FName(TEXT("VisualLogger")));
 	}
+
+	/**
+	* Opens the merge actors tab.
+	* Remove this when it is no longer experimental.
+	*/
+	static void OpenMergeActors()
+	{
+		FGlobalTabmanager::Get()->InvokeTab(FName(TEXT("MergeActorsApp")));
+	}
 };
diff --git a/Engine/Source/Editor/MergeActors/Private/MergeActorsModule.cpp b/Engine/Source/Editor/MergeActors/Private/MergeActorsModule.cpp
--- a/Engine/Source/Editor/MergeActors/Private/MergeActorsModule.cpp
+++ b/Engine/Source/Editor/MergeActors/Private/MergeActorsModule.cpp
@@ -95,12 +95,9 @@ void FMergeActorsModule::StartupModule()
 {
 	bRegisteredTabSpawner = false;
 
-	if (!GetDefault<UEditorExperimentalSettings>()->bActorMerging)
-	{
-		return;
-	}
-
-	FGlobalTabmanager::Get()->RegisterNomadTabSpawner(MergeActorsApp, FOnSpawnTab::CreateRaw(this, &FMergeActorsModule::CreateMergeActorsTab))
+	// This is still experimental in the editor, so it's invoked specifically in FMainMenu for now.
+	// When no longer experimental, switch to the nomad spawner registration below
+	FGlobalTabmanager::Get()->RegisterTabSpawner(MergeActorsApp, FOnSpawnTab::CreateRaw(this, &FMergeActorsModule::CreateMergeActorsTab))
 		.SetDisplayName(LOCTEXT("TabTitle", "Merge Actors"))
 		.SetTooltipText(LOCTEXT("TooltipText", "Open the Merge Actors tab."))
 		.SetGroup(WorkspaceMenu::GetMenuStructure().GetLevelEditorCategory())
@@ -110,7 +107,13 @@ void FMergeActorsModule::StartupModule()
 
 	// Register built-in merging tools straight away
 	ensure(RegisterMergeActorsTool(MakeUnique<FMeshMergingTool>()));
-	ensure(RegisterMergeActorsTool(MakeUnique<FMeshProxyTool>()));
+
+	IMeshUtilities* MeshUtilities = FModuleManager::Get().LoadModulePtr<IMeshUtilities>("MeshUtilities");
+	if (MeshUtilities != nullptr && MeshUtilities->GetMeshMergingInterface() != nullptr)
+	{
+		// Only register MeshProxyTool if Simplygon is available
+		ensure(RegisterMergeActorsTool(MakeUnique<FMeshProxyTool>()));
+	}
 }
 
 
@@ -123,7 +126,7 @@ void FMergeActorsModule::ShutdownModule()
 
 	if (FSlateApplication::IsInitialized())
 	{
-		FGlobalTabmanager::Get()->UnregisterNomadTabSpawner(MergeActorsApp);
+		FGlobalTabmanager::Get()->UnregisterTabSpawner(MergeActorsApp);
 	}
 }
 
