Fixed: Merge Actors "place in world" option does not work
UE-15119

diff --git a/Engine/Source/Editor/MergeActors/Private/MeshMergingTool/MeshMergingTool.cpp b/Engine/Source/Editor/MergeActors/Private/MeshMergingTool/MeshMergingTool.cpp
--- a/Engine/Source/Editor/MergeActors/Private/MeshMergingTool/MeshMergingTool.cpp
+++ b/Engine/Source/Editor/MergeActors/Private/MeshMergingTool/MeshMergingTool.cpp
@@ -111,25 +111,29 @@ bool FMeshMergingTool::RunMerge(const FString& PackageName)
 		// Place new mesh in the world
 		if (bPlaceInWorld)
 		{
-			const FScopedTransaction Transaction(NSLOCTEXT("UnrealEd", "PlaceMergedActor", "Place Merged Actor"));
-			UniqueLevels[0]->Modify();
-
-			UWorld* World = UniqueLevels[0]->OwningWorld;
-			FActorSpawnParameters Params;
-			Params.OverrideLevel = UniqueLevels[0];
-			FRotator MergedActorRotation(ForceInit);
-
-			AStaticMeshActor* MergedActor = World->SpawnActor<AStaticMeshActor>(MergedActorLocation, MergedActorRotation, Params);
-			MergedActor->GetStaticMeshComponent()->StaticMesh = Cast<UStaticMesh>(AssetsToSync[0]);
-			MergedActor->SetActorLabel(AssetsToSync[0]->GetName());
-
-			// Add source actors as children to merged actor and hide them
-			for (AActor* Actor : Actors)
+			UStaticMesh* MergedMesh = nullptr;
+			if (AssetsToSync.FindItemByClass(&MergedMesh))
 			{
-				Actor->Modify();
-				Actor->AttachRootComponentToActor(MergedActor, NAME_None, EAttachLocation::KeepWorldPosition);
-				Actor->SetActorHiddenInGame(true);
-				Actor->SetIsTemporarilyHiddenInEditor(true);
+				const FScopedTransaction Transaction(NSLOCTEXT("UnrealEd", "PlaceMergedActor", "Place Merged Actor"));
+				UniqueLevels[0]->Modify();
+
+				UWorld* World = UniqueLevels[0]->OwningWorld;
+				FActorSpawnParameters Params;
+				Params.OverrideLevel = UniqueLevels[0];
+				FRotator MergedActorRotation(ForceInit);
+								
+				AStaticMeshActor* MergedActor = World->SpawnActor<AStaticMeshActor>(MergedActorLocation, MergedActorRotation, Params);
+				MergedActor->GetStaticMeshComponent()->StaticMesh = MergedMesh;
+				MergedActor->SetActorLabel(AssetsToSync[0]->GetName());
+
+				// Add source actors as children to merged actor and hide them
+				for (AActor* Actor : Actors)
+				{
+					Actor->Modify();
+					Actor->AttachRootComponentToActor(MergedActor, NAME_None, EAttachLocation::KeepWorldPosition);
+					Actor->SetActorHiddenInGame(true);
+					Actor->SetIsTemporarilyHiddenInEditor(true);
+				}
 			}
 		}
 	}
