Fixed: Merge Actors produces incorrect merged material at resolution > 1024
UE-15118

diff --git a/Engine/Source/Editor/MergeActors/Private/MeshMergingTool/SMeshMergingDialog.cpp b/Engine/Source/Editor/MergeActors/Private/MeshMergingTool/SMeshMergingDialog.cpp
--- a/Engine/Source/Editor/MergeActors/Private/MeshMergingTool/SMeshMergingDialog.cpp
+++ b/Engine/Source/Editor/MergeActors/Private/MeshMergingTool/SMeshMergingDialog.cpp
@@ -27,8 +27,8 @@ void SMeshMergingDialog::Construct(const FArguments& InArgs, FMeshMergingTool* I
 	const int32 MaxTexResolution = 1 << FTextureLODGroup().MaxLODMipCount;
 	for (int32 TexRes = MinTexResolution; TexRes <= MaxTexResolution; TexRes*=2)
 	{
-		LightMapResolutionOptions.Add(MakeShareable(new FString(FString::FormatAsNumber(TexRes))));
-		MergedMaterialResolutionOptions.Add(MakeShareable(new FString(FString::FormatAsNumber(TexRes))));
+		LightMapResolutionOptions.Add(MakeShareable(new FString(TTypeToString<int32>::ToString(TexRes))));
+		MergedMaterialResolutionOptions.Add(MakeShareable(new FString(TTypeToString<int32>::ToString(TexRes))));
 	}
 
 	Tool->MergingSettings.TargetLightMapResolution = FMath::Clamp(Tool->MergingSettings.TargetLightMapResolution, MinTexResolution, MaxTexResolution);
@@ -37,12 +37,12 @@ void SMeshMergingDialog::Construct(const FArguments& InArgs, FMeshMergingTool* I
 	// Setup available UV channels for an atlased lightmap
 	for (int32 Index = 0; Index < MAX_MESH_TEXTURE_COORDS; Index++)
 	{
-		LightMapChannelOptions.Add(MakeShareable(new FString(FString::FormatAsNumber(Index))));
+		LightMapChannelOptions.Add(MakeShareable(new FString(TTypeToString<int32>::ToString(Index))));
 	}
 
 	for (int32 Index = 0; Index < MAX_STATIC_MESH_LODS; Index++)
 	{
-		ExportLODOptions.Add(MakeShareable(new FString(FString::FormatAsNumber(Index))));
+		ExportLODOptions.Add(MakeShareable(new FString(TTypeToString<int32>::ToString(Index))));
 	}
 
 	// Create widget layout
