Changed MergeActors tool to use a Nomad Tab, so that any open tab can be correctly restored when reopening the editor.

diff --git a/Engine/Source/Editor/MainFrame/Private/Menus/MainMenu.h b/Engine/Source/Editor/MainFrame/Private/Menus/MainMenu.h
--- a/Engine/Source/Editor/MainFrame/Private/Menus/MainMenu.h
+++ b/Engine/Source/Editor/MainFrame/Private/Menus/MainMenu.h
@@ -118,6 +118,6 @@ private:
 	*/
 	static void OpenMergeActors()
 	{
-		FGlobalTabmanager::Get()->InvokeTab(FName(TEXT("MergeActorsApp")));
+		FGlobalTabmanager::Get()->InvokeTab(FName(TEXT("MergeActors")));
 	}
 };
diff --git a/Engine/Source/Editor/MergeActors/Private/MergeActorsModule.cpp b/Engine/Source/Editor/MergeActors/Private/MergeActorsModule.cpp
--- a/Engine/Source/Editor/MergeActors/Private/MergeActorsModule.cpp
+++ b/Engine/Source/Editor/MergeActors/Private/MergeActorsModule.cpp
@@ -11,7 +11,7 @@
 #define LOCTEXT_NAMESPACE "MergeActorsModule"
 
 
-static const FName MergeActorsApp = FName("MergeActorsApp");
+static const FName MergeActorsTabName = FName("MergeActors");
 
 /**
  * Merge Actors module
@@ -53,9 +53,6 @@ private:
 
 	/** List of registered MergeActorsTool instances */
 	TArray<TUniquePtr<IMergeActorsTool>> MergeActorsTools;
-
-	/** Whether a nomad tab spawner was registered */
-	bool bRegisteredTabSpawner;
 };
 
 IMPLEMENT_MODULE(FMergeActorsModule, MergeActors);
@@ -93,17 +90,14 @@ TSharedRef<SDockTab> FMergeActorsModule::CreateMergeActorsTab(const FSpawnTabArg
 
 void FMergeActorsModule::StartupModule()
 {
-	bRegisteredTabSpawner = false;
-
-	// This is still experimental in the editor, so it's invoked specifically in FMainMenu for now.
-	// When no longer experimental, switch to the nomad spawner registration below
-	FGlobalTabmanager::Get()->RegisterTabSpawner(MergeActorsApp, FOnSpawnTab::CreateRaw(this, &FMergeActorsModule::CreateMergeActorsTab))
+	FGlobalTabmanager::Get()->RegisterNomadTabSpawner(MergeActorsTabName, FOnSpawnTab::CreateRaw(this, &FMergeActorsModule::CreateMergeActorsTab))
 		.SetDisplayName(LOCTEXT("TabTitle", "Merge Actors"))
 		.SetTooltipText(LOCTEXT("TooltipText", "Open the Merge Actors tab."))
 		.SetGroup(WorkspaceMenu::GetMenuStructure().GetLevelEditorCategory())
-		.SetIcon(FSlateIcon(FEditorStyle::GetStyleSetName(), "ClassViewer.TabIcon"));
-
-	bRegisteredTabSpawner = true;
+		.SetIcon(FSlateIcon(FEditorStyle::GetStyleSetName(), "ClassViewer.TabIcon"))
+		// This is still experimental in the editor, so it's added specifically in FMainMenu for now.
+		// When no longer experimental, remove the below.
+		.SetAutoGenerateMenuEntry(false);
 
 	// Register built-in merging tools straight away
 	ensure(RegisterMergeActorsTool(MakeUnique<FMeshMergingTool>()));
@@ -119,14 +113,9 @@ void FMergeActorsModule::StartupModule()
 
 void FMergeActorsModule::ShutdownModule()
 {
-	if (!bRegisteredTabSpawner)
-	{
-		return;
-	}
-
 	if (FSlateApplication::IsInitialized())
 	{
-		FGlobalTabmanager::Get()->UnregisterTabSpawner(MergeActorsApp);
+		FGlobalTabmanager::Get()->UnregisterNomadTabSpawner(MergeActorsTabName);
 	}
 }
 
