Added: Lightmap resolution option for MeshProxy tool
Added: Missing tooltips to MeshProxy options

diff --git a/Engine/Source/Developer/MeshUtilities/Private/MeshUtilities.cpp b/Engine/Source/Developer/MeshUtilities/Private/MeshUtilities.cpp
--- a/Engine/Source/Developer/MeshUtilities/Private/MeshUtilities.cpp
+++ b/Engine/Source/Developer/MeshUtilities/Private/MeshUtilities.cpp
@@ -3864,7 +3864,7 @@ void FMeshUtilities::CreateProxyMesh(
 		StaticMesh->LightingGuid = FGuid::NewGuid();
 
 		// Set it to use textured lightmaps. Note that Build Lighting will do the error-checking (texcoordindex exists for all LODs, etc).
-		StaticMesh->LightMapResolution = 64;
+		StaticMesh->LightMapResolution = InProxySettings.LightMapResolution;
 		StaticMesh->LightMapCoordinateIndex = 1;
 
 		FStaticMeshSourceModel* SrcModel = new (StaticMesh->SourceModels) FStaticMeshSourceModel();
diff --git a/Engine/Source/Developer/SimplygonMeshReduction/Private/SimplygonMeshReduction.cpp b/Engine/Source/Developer/SimplygonMeshReduction/Private/SimplygonMeshReduction.cpp
--- a/Engine/Source/Developer/SimplygonMeshReduction/Private/SimplygonMeshReduction.cpp
+++ b/Engine/Source/Developer/SimplygonMeshReduction/Private/SimplygonMeshReduction.cpp
@@ -505,9 +505,9 @@ public:
 			RemeshingProcessor->GetRemeshingSettings()->SetGroundPlaneNegativeHalfspace(InProxySettings.bPlaneNegativeHalfspace);
 		}
 		
-		RemeshingProcessor->GetRemeshingSettings()->SetTransferNormals( false );
-		RemeshingProcessor->GetRemeshingSettings()->SetMergeDistance( InProxySettings.MergeDistance );
-		RemeshingProcessor->GetRemeshingSettings()->SetHardEdgeAngle( FMath::DegreesToRadians(InProxySettings.HardAngleThreshold) );//This should be a user settings in the popup dialog!
+		RemeshingProcessor->GetRemeshingSettings()->SetTransferNormals(InProxySettings.bRecalculateNormals == false);
+		RemeshingProcessor->GetRemeshingSettings()->SetMergeDistance(InProxySettings.MergeDistance);
+		RemeshingProcessor->GetRemeshingSettings()->SetHardEdgeAngle(FMath::DegreesToRadians(InProxySettings.HardAngleThreshold));
 		RemeshingProcessor->SetSceneRoot(Scene->GetRootNode());
 
 		//Setup the mapping image used for casting
diff --git a/Engine/Source/Editor/MergeActors/Private/MeshProxyTool/SMeshProxyDialog.cpp b/Engine/Source/Editor/MergeActors/Private/MeshProxyTool/SMeshProxyDialog.cpp
--- a/Engine/Source/Editor/MergeActors/Private/MeshProxyTool/SMeshProxyDialog.cpp
+++ b/Engine/Source/Editor/MergeActors/Private/MeshProxyTool/SMeshProxyDialog.cpp
@@ -22,6 +22,7 @@ void SMeshProxyDialog::Construct(const FArguments& InArgs, FMeshProxyTool* InToo
 	CuttingPlaneOptions.Add(MakeShareable(new FString(TEXT("-Y"))));
 	CuttingPlaneOptions.Add(MakeShareable(new FString(TEXT("-Z"))));
 
+	TextureResolutionOptions.Add(MakeShareable(new FString(TEXT("32"))));
 	TextureResolutionOptions.Add(MakeShareable(new FString(TEXT("64"))));
 	TextureResolutionOptions.Add(MakeShareable(new FString(TEXT("128"))));
 	TextureResolutionOptions.Add(MakeShareable(new FString(TEXT("256"))));
@@ -35,6 +36,11 @@ void SMeshProxyDialog::Construct(const FArguments& InArgs, FMeshProxyTool* InToo
 BEGIN_SLATE_FUNCTION_BUILD_OPTIMIZATION
 void SMeshProxyDialog::CreateLayout()
 {
+	int32 TextureResEntryIndex = FindTextureResolutionEntryIndex(Tool->ProxySettings.Material.BaseColorMapSize.X);
+	int32 LightMapResEntryIndex = FindTextureResolutionEntryIndex(Tool->ProxySettings.LightMapResolution);
+	TextureResEntryIndex = FMath::Max(TextureResEntryIndex, 0);
+	LightMapResEntryIndex = FMath::Max(LightMapResEntryIndex, 0);
+		
 	this->ChildSlot
 	[
 		SNew(SVerticalBox)
@@ -71,6 +77,7 @@ void SMeshProxyDialog::CreateLayout()
 						SNew(STextBlock)
 						.Text(LOCTEXT("OnScreenSizeLabel", "On Screen Size (pixels)"))
 						.Font(FEditorStyle::GetFontStyle("StandardDialog.SmallFont"))
+						.ToolTipText(GetPropertyToolTipText(GET_MEMBER_NAME_CHECKED(FMeshProxySettings, ScreenSize)))
 					]
 					+ SHorizontalBox::Slot()
 					.FillWidth(0.5f)
@@ -107,6 +114,7 @@ void SMeshProxyDialog::CreateLayout()
 						SNew(STextBlock)
 						.Text(LOCTEXT("MergeDistanceLabel", "Merge Distance (pixels)"))
 						.Font(FEditorStyle::GetFontStyle("StandardDialog.SmallFont"))
+						.ToolTipText(GetPropertyToolTipText(GET_MEMBER_NAME_CHECKED(FMeshProxySettings, MergeDistance)))
 					]
 					+ SHorizontalBox::Slot()
 					.FillWidth(0.5f)
@@ -152,7 +160,7 @@ void SMeshProxyDialog::CreateLayout()
 						SNew(STextComboBox)
 						.Font(FEditorStyle::GetFontStyle("StandardDialog.SmallFont"))
 						.OptionsSource(&TextureResolutionOptions)
-						.InitiallySelectedItem(TextureResolutionOptions[3]) //512
+						.InitiallySelectedItem(TextureResolutionOptions[TextureResEntryIndex])
 						.OnSelectionChanged(this, &SMeshProxyDialog::SetTextureResolution)
 					]
 				]
@@ -161,15 +169,26 @@ void SMeshProxyDialog::CreateLayout()
 				.AutoHeight()
 				.Padding(FEditorStyle::GetMargin("StandardDialog.ContentPadding"))
 				[
-					SNew(SCheckBox)
-					.Type(ESlateCheckBoxType::CheckBox)
-					.IsChecked(this, &SMeshProxyDialog::GetRecalculateNormals)
-					.OnCheckStateChanged(this, &SMeshProxyDialog::SetRecalculateNormals)
-					.Content()
+					SNew(SHorizontalBox)
+					+SHorizontalBox::Slot()
+					.FillWidth(0.5f)
+					.VAlign(VAlign_Center)
 					[
 						SNew(STextBlock)
-						.Text(LOCTEXT("RecalcNormalsLabel", "Recalculate Normals"))
+						.Text(LOCTEXT("LightMapResolutionLabel", "LightMap Resolution"))
+						.Font(FEditorStyle::GetFontStyle("StandardDialog.SmallFont"))
+						.ToolTipText(GetPropertyToolTipText(GET_MEMBER_NAME_CHECKED(FMeshProxySettings, LightMapResolution)))
+					]
+					+ SHorizontalBox::Slot()
+					.FillWidth(0.5f)
+					.HAlign(HAlign_Left)
+					.VAlign(VAlign_Center)
+					[
+						SNew(STextComboBox)
 						.Font(FEditorStyle::GetFontStyle("StandardDialog.SmallFont"))
+						.OptionsSource(&TextureResolutionOptions)
+						.InitiallySelectedItem(TextureResolutionOptions[LightMapResEntryIndex])
+						.OnSelectionChanged(this, &SMeshProxyDialog::SetLightMapResolution)
 					]
 				]
 
@@ -183,25 +202,10 @@ void SMeshProxyDialog::CreateLayout()
 					.VAlign(VAlign_Center)
 					.Padding(0.0, 0.0, 3.0, 0.0)
 					[
-						SNew(SHorizontalBox)
-						+SHorizontalBox::Slot()
-						.AutoWidth()
-						[
-							SNew(SBox)
-							.MinDesiredWidth(30)
-							[
-								SNullWidget::NullWidget
-							]
-						]
-						+SHorizontalBox::Slot()
-						.AutoWidth()
-						.HAlign(HAlign_Left)
-						[
-							SNew(STextBlock)
-							.IsEnabled(this, &SMeshProxyDialog::HardAngleThresholdEnabled)
-							.Text(LOCTEXT("HardAngleLabel", "Hard Angle"))
-							.Font(FEditorStyle::GetFontStyle("StandardDialog.SmallFont"))
-						]
+						SNew(STextBlock)
+						.Text(LOCTEXT("HardAngleLabel", "Hard Edge Angle"))
+						.Font(FEditorStyle::GetFontStyle("StandardDialog.SmallFont"))
+						.ToolTipText(GetPropertyToolTipText(GET_MEMBER_NAME_CHECKED(FMeshProxySettings, HardAngleThreshold)))
 					]
 					+SHorizontalBox::Slot()
 					.FillWidth(0.5f)
@@ -227,6 +231,23 @@ void SMeshProxyDialog::CreateLayout()
 					]
 				]
 
+				+SVerticalBox::Slot()
+				.AutoHeight()
+				.Padding(FEditorStyle::GetMargin("StandardDialog.ContentPadding"))
+				[
+					SNew(SCheckBox)
+					.Type(ESlateCheckBoxType::CheckBox)
+					.IsChecked(this, &SMeshProxyDialog::GetRecalculateNormals)
+					.OnCheckStateChanged(this, &SMeshProxyDialog::SetRecalculateNormals)
+					.Content()
+					[
+						SNew(STextBlock)
+						.Text(LOCTEXT("RecalcNormalsLabel", "Recalculate Normals"))
+						.Font(FEditorStyle::GetFontStyle("StandardDialog.SmallFont"))
+						.ToolTipText(GetPropertyToolTipText(GET_MEMBER_NAME_CHECKED(FMeshProxySettings, bRecalculateNormals)))
+					]
+				]
+
 				+SVerticalBox::Slot()
 				.AutoHeight()
 				.Padding(FEditorStyle::GetMargin("StandardDialog.ContentPadding"))
@@ -240,6 +261,7 @@ void SMeshProxyDialog::CreateLayout()
 						SNew(STextBlock)
 						.Text(LOCTEXT("ClippingPlaneLabel", "Use Clipping Plane"))
 						.Font(FEditorStyle::GetFontStyle("StandardDialog.SmallFont"))
+						.ToolTipText(GetPropertyToolTipText(GET_MEMBER_NAME_CHECKED(FMeshProxySettings, bUseClippingPlane)))
 					]
 				]
 
@@ -270,6 +292,7 @@ void SMeshProxyDialog::CreateLayout()
 							.IsEnabled(this, &SMeshProxyDialog::UseClippingPlaneEnabled)
 							.Text(LOCTEXT("ClippingAxisLabel", "Clipping Axis"))
 							.Font(FEditorStyle::GetFontStyle("StandardDialog.SmallFont"))
+							.ToolTipText(GetPropertyToolTipText(GET_MEMBER_NAME_CHECKED(FMeshProxySettings, AxisIndex)))
 						]
 					]
 					+SHorizontalBox::Slot()
@@ -315,6 +338,7 @@ void SMeshProxyDialog::CreateLayout()
 							.Text(LOCTEXT("PlaneLevelLabel", "Plane level"))
 							.Font(FEditorStyle::GetFontStyle("StandardDialog.SmallFont"))
 							.IsEnabled(this, &SMeshProxyDialog::UseClippingPlaneEnabled)
+							.ToolTipText(GetPropertyToolTipText(GET_MEMBER_NAME_CHECKED(FMeshProxySettings, ClippingLevel)))
 						]
 					]
 					+SHorizontalBox::Slot()
@@ -404,6 +428,29 @@ void SMeshProxyDialog::CreateLayout()
 }
 END_SLATE_FUNCTION_BUILD_OPTIMIZATION
 
+int32 SMeshProxyDialog::FindTextureResolutionEntryIndex(int32 InResolution) const
+{
+	FString ResolutionStr = TTypeToString<int32>::ToString(InResolution);
+	
+	int32 Result = TextureResolutionOptions.IndexOfByPredicate([&](const TSharedPtr<FString>& Entry)
+	{
+		return (ResolutionStr == *Entry);
+	});
+
+	return Result;
+}
+
+FText SMeshProxyDialog::GetPropertyToolTipText(const FName& PropertyName) const
+{
+	UProperty* Property = FMeshProxySettings::StaticStruct()->FindPropertyByName(PropertyName);
+	if (Property)
+	{
+		return Property->GetToolTipText();
+	}
+	
+	return FText::GetEmpty();
+}
+
 //Screen size
 TOptional<int32> SMeshProxyDialog::GetScreenSize() const
 {
@@ -521,6 +568,14 @@ void SMeshProxyDialog::SetTextureResolution(TSharedPtr<FString> NewSelection, ES
 	Tool->ProxySettings.Material.SpecularMapSize = TextureSize;
 }
 
+void SMeshProxyDialog::SetLightMapResolution(TSharedPtr<FString> NewSelection, ESelectInfo::Type SelectInfo)
+{
+	int32 Resolution = 256;
+	TTypeFromString<int32>::FromString(Resolution, **NewSelection);
+		
+	Tool->ProxySettings.LightMapResolution = Resolution;
+}
+
 ECheckBoxState SMeshProxyDialog::GetExportNormalMap() const
 {
 	return Tool->ProxySettings.Material.bNormalMap ? ECheckBoxState::Checked :  ECheckBoxState::Unchecked;
diff --git a/Engine/Source/Editor/MergeActors/Private/MeshProxyTool/SMeshProxyDialog.h b/Engine/Source/Editor/MergeActors/Private/MeshProxyTool/SMeshProxyDialog.h
--- a/Engine/Source/Editor/MergeActors/Private/MeshProxyTool/SMeshProxyDialog.h
+++ b/Engine/Source/Editor/MergeActors/Private/MeshProxyTool/SMeshProxyDialog.h
@@ -47,6 +47,7 @@ protected:
 
 	/** TextureResolution accessors */
 	void SetTextureResolution(TSharedPtr<FString> NewSelection, ESelectInfo::Type SelectInfo);
+	void SetLightMapResolution(TSharedPtr<FString> NewSelection, ESelectInfo::Type SelectInfo);
 
 	/** Export material properties acessors **/
 	ECheckBoxState GetExportNormalMap() const;
@@ -62,6 +63,9 @@ private:
 	/** Creates the geometry mode controls */
 	void CreateLayout();
 
+	int32 FindTextureResolutionEntryIndex(int32 InResolution) const;
+	FText GetPropertyToolTipText(const FName& PropertyName) const;
+
 private:
 	FMeshProxyTool* Tool;
 
diff --git a/Engine/Source/Runtime/Engine/Classes/Engine/EngineTypes.h b/Engine/Source/Runtime/Engine/Classes/Engine/EngineTypes.h
--- a/Engine/Source/Runtime/Engine/Classes/Engine/EngineTypes.h
+++ b/Engine/Source/Runtime/Engine/Classes/Engine/EngineTypes.h
@@ -2368,26 +2368,35 @@ struct FMeshProxySettings
 	UPROPERTY()
 	bool bExportSpecularMap_DEPRECATED;
 
-	/** Should Simplygon recalculate normals for the proxy mesh? */
+	/** Lightmap resolution */
 	UPROPERTY(EditAnywhere, Category=ProxySettings)
-	bool bRecalculateNormals;
+	int32 LightMapResolution;
 
-	/** Angle at which a hard edge is introduced between faces. */
+	/** Whether Simplygon should recalculate normals, otherwise the normals channel will be sampled from the original mesh */
 	UPROPERTY(EditAnywhere, Category=ProxySettings)
+	bool bRecalculateNormals;
+
+	/** Angle at which a hard edge is introduced between faces */
+	UPROPERTY(EditAnywhere, Category=ProxySettings, meta=(DisplayName="Hard Edge Angle"))
 	float HardAngleThreshold;
 
+	/** Set the on-screen merge distance in pixels. Smaller cavities will be removed */
 	UPROPERTY(EditAnywhere, Category=ProxySettings)
 	int32 MergeDistance;
-
+	
+	/** Set to true to cap the mesh with a ground plane */
 	UPROPERTY(EditAnywhere, Category=ProxySettings)
 	bool bUseClippingPlane;
 
+	/* Ground plane level */
 	UPROPERTY(EditAnywhere, Category=ProxySettings)
 	float ClippingLevel;
 
+	/** Set the axis index for the ground plane (0:X-Axis, 1:Y-Axis, 2:Z-Axis) */
 	UPROPERTY(EditAnywhere, Category=ProxySettings)
 	int32 AxisIndex;
 
+	/** Set to true to use negative halfspace for model, and reject the positive halfspace */
 	UPROPERTY(EditAnywhere, Category=ProxySettings)
 	bool bPlaneNegativeHalfspace;
 
@@ -2400,6 +2409,7 @@ struct FMeshProxySettings
 		, bExportMetallicMap_DEPRECATED(false)
 		, bExportRoughnessMap_DEPRECATED(false)
 		, bExportSpecularMap_DEPRECATED(false)
+		, LightMapResolution(256)
 		, bRecalculateNormals(true)
 		, HardAngleThreshold(80.0f)
 		, MergeDistance(4)
