Mesh Merging  dialog - missing notification when no meshes are selected.

diff --git a/Engine/Source/Editor/MergeActors/Private/MergeProxyUtils/Utils.cpp b/Engine/Source/Editor/MergeActors/Private/MergeProxyUtils/Utils.cpp
--- a/Engine/Source/Editor/MergeActors/Private/MergeProxyUtils/Utils.cpp
+++ b/Engine/Source/Editor/MergeActors/Private/MergeProxyUtils/Utils.cpp
@@ -135,6 +135,13 @@ TSharedRef<ITableRow> FComponentSelectionControl::MakeComponentListItemWidget(TS
 		State = *StoredState;
 	}
 
+	if (UStaticMeshComponent* StaticMeshComponent = Cast<UStaticMeshComponent>(ComponentData->PrimComponent.Get()))
+	{
+		bEnabled = (StaticMeshComponent->GetStaticMesh() != nullptr);
+		bIsMesh = true;
+	}
+
+
 	return SNew(STableRow<TSharedPtr<FMergeComponentData>>, OwnerTable)
 		[
 			SNew(SBox)
diff --git a/Engine/Source/Editor/MergeActors/Private/MeshProxyTool/SMeshProxyDialog.cpp b/Engine/Source/Editor/MergeActors/Private/MeshProxyTool/SMeshProxyDialog.cpp
--- a/Engine/Source/Editor/MergeActors/Private/MeshProxyTool/SMeshProxyDialog.cpp
+++ b/Engine/Source/Editor/MergeActors/Private/MeshProxyTool/SMeshProxyDialog.cpp
@@ -116,7 +116,7 @@ void  SMeshProxyDialog::Construct(const FArguments& InArgs, FMeshProxyTool* InTo
 			.Visibility_Lambda([this]()->EVisibility { return this->GetContentEnabledState() ? EVisibility::Collapsed : EVisibility::Visible; })
 			[
 				SNew(STextBlock)
-				.Text(LOCTEXT("DeleteUndo", "Insufficient mesh components found for merging."))
+				.Text(LOCTEXT("DeleteUndo", "Insufficient mesh components found for ProxyLOD merging."))
 			]
 		]
 	];
