Moving over CL 6926865:
"Fix mesh tool crash at exit

diff --git a/Engine/Source/Editor/MergeActors/Private/MeshInstancingTool/MeshInstancingTool.h b/Engine/Source/Editor/MergeActors/Private/MeshInstancingTool/MeshInstancingTool.h
--- a/Engine/Source/Editor/MergeActors/Private/MeshInstancingTool/MeshInstancingTool.h
+++ b/Engine/Source/Editor/MergeActors/Private/MeshInstancingTool/MeshInstancingTool.h
@@ -38,7 +38,7 @@ public:
 	{
 		if (bInitialized)
 		{
-			if (DefaultSettings)
+			if (UObjectInitialized() && DefaultSettings)
 			{
 				DefaultSettings->RemoveFromRoot();
 				DefaultSettings->MarkPendingKill();
diff --git a/Engine/Source/Editor/MergeActors/Private/MeshMergingTool/MeshMergingTool.h b/Engine/Source/Editor/MergeActors/Private/MeshMergingTool/MeshMergingTool.h
--- a/Engine/Source/Editor/MergeActors/Private/MeshMergingTool/MeshMergingTool.h
+++ b/Engine/Source/Editor/MergeActors/Private/MeshMergingTool/MeshMergingTool.h
@@ -44,7 +44,7 @@ public:
 	{
 		if (bInitialized)
 		{
-			if (DefaultSettings)
+			if (UObjectInitialized() && DefaultSettings)
 			{
 				DefaultSettings->RemoveFromRoot();
 				DefaultSettings->MarkPendingKill();
diff --git a/Engine/Source/Editor/MergeActors/Private/MeshProxyTool/MeshProxyTool.h b/Engine/Source/Editor/MergeActors/Private/MeshProxyTool/MeshProxyTool.h
--- a/Engine/Source/Editor/MergeActors/Private/MeshProxyTool/MeshProxyTool.h
+++ b/Engine/Source/Editor/MergeActors/Private/MeshProxyTool/MeshProxyTool.h
@@ -36,7 +36,7 @@ public:
 	{
 		if (bInitialized)
 		{
-			if (DefaultSettings)
+			if (UObjectInitialized() && DefaultSettings)
 			{
 				DefaultSettings->RemoveFromRoot();
 				DefaultSettings->MarkPendingKill();
