Merge Actors tool leads to crash when the selected list is not matching current existing objects due to construction script.
#fix  Ensure that we validate that Primitive Component weakptr is still valid when applying the tool
#misc Updated listview text to be dynamic and show 'Invalid Actor' for actors which have been destroyed / are invalidated

diff --git a/Engine/Source/Editor/MergeActors/Private/MergeProxyUtils/Utils.cpp b/Engine/Source/Editor/MergeActors/Private/MergeProxyUtils/Utils.cpp
--- a/Engine/Source/Editor/MergeActors/Private/MergeProxyUtils/Utils.cpp
+++ b/Engine/Source/Editor/MergeActors/Private/MergeProxyUtils/Utils.cpp
@@ -183,7 +183,27 @@ TSharedRef<ITableRow> FComponentSelectionControl::MakeComponentListItemWidget(TS
 				.AutoWidth()
 				[
 					SNew(STextBlock)
-					.Text(FText::FromString(OwningActorName + " - " + ComponentInfo + " - " + ComponentName))
+					.Text_Lambda([ComponentData]() -> FText
+					{
+						if (ComponentData.IsValid() && ComponentData->PrimComponent.IsValid())
+						{
+							const FString OwningActorName = ComponentData->PrimComponent->GetOwner()->GetName();
+							const FString ComponentName = ComponentData->PrimComponent->GetName();
+							FString ComponentInfo;
+							if (UStaticMeshComponent* StaticMeshComponent = Cast<UStaticMeshComponent>(ComponentData->PrimComponent.Get()))
+							{
+								ComponentInfo = (StaticMeshComponent->GetStaticMesh() != nullptr) ? StaticMeshComponent->GetStaticMesh()->GetName() : TEXT("No Static Mesh Available");
+							}
+							else if (UShapeComponent* ShapeComponent = Cast<UShapeComponent>(ComponentData->PrimComponent.Get()))
+							{
+								ComponentInfo = ShapeComponent->GetClass()->GetName();
+							}
+
+							return FText::FromString(OwningActorName + " - " + ComponentInfo + " - " + ComponentName);
+						}
+
+						return FText::FromString("Invalid Actor");
+					})
 				]
 			]
 		];
diff --git a/Engine/Source/Editor/MergeActors/Private/MeshInstancingTool/MeshInstancingTool.cpp b/Engine/Source/Editor/MergeActors/Private/MeshInstancingTool/MeshInstancingTool.cpp
--- a/Engine/Source/Editor/MergeActors/Private/MeshInstancingTool/MeshInstancingTool.cpp
+++ b/Engine/Source/Editor/MergeActors/Private/MeshInstancingTool/MeshInstancingTool.cpp
@@ -81,16 +81,19 @@ bool FMeshInstancingTool::RunMerge(const FString& PackageName)
 		for ( const TSharedPtr<FInstanceComponentData>& SelectedComponent : SelectedComponents)
 		{
 			// Determine whether or not this component should be incorporated according the user settings
-			if (SelectedComponent->bShouldIncorporate)
+			if (SelectedComponent->bShouldIncorporate && SelectedComponent->PrimComponent.IsValid())
 			{
 				ComponentsToMerge.Add(SelectedComponent->PrimComponent.Get());
 			}
 		}
-		
-		// spawn the actor that will contain out instances
-		UWorld* World = ComponentsToMerge[0]->GetWorld();
-		checkf(World != nullptr, TEXT("Invalid World retrieved from Mesh components"));
-		MeshUtilities.MergeComponentsToInstances(ComponentsToMerge, World, UniqueLevels[0], SettingsObject->Settings);
+
+		if (ComponentsToMerge.Num())
+		{
+			// spawn the actor that will contain out instances
+			UWorld* World = ComponentsToMerge[0]->GetWorld();
+			checkf(World != nullptr, TEXT("Invalid World retrieved from Mesh components"));
+			MeshUtilities.MergeComponentsToInstances(ComponentsToMerge, World, UniqueLevels[0], SettingsObject->Settings);
+		}
 	}
 
 	InstancingDialog->Reset();
diff --git a/Engine/Source/Editor/MergeActors/Private/MeshMergingTool/MeshMergingTool.cpp b/Engine/Source/Editor/MergeActors/Private/MeshMergingTool/MeshMergingTool.cpp
--- a/Engine/Source/Editor/MergeActors/Private/MeshMergingTool/MeshMergingTool.cpp
+++ b/Engine/Source/Editor/MergeActors/Private/MeshMergingTool/MeshMergingTool.cpp
@@ -105,16 +105,19 @@ bool FMeshMergingTool::RunMerge(const FString& PackageName)
 		for ( const TSharedPtr<FMergeComponentData>& SelectedComponent : SelectedComponents)
 		{
 			// Determine whether or not this component should be incorporated according the user settings
-			if (SelectedComponent->bShouldIncorporate)
+			if (SelectedComponent->bShouldIncorporate && SelectedComponent->PrimComponent.IsValid())
 			{
 				ComponentsToMerge.Add(SelectedComponent->PrimComponent.Get());
 			}
 		}
-		
-		UWorld* World = ComponentsToMerge[0]->GetWorld();	
-		checkf(World != nullptr, TEXT("Invalid World retrieved from Mesh components"));
-		const float ScreenAreaSize = TNumericLimits<float>::Max();
-		MeshUtilities.MergeComponentsToStaticMesh(ComponentsToMerge, World, SettingsObject->Settings, nullptr, nullptr, PackageName, AssetsToSync, MergedActorLocation, ScreenAreaSize, true);
+
+		if (ComponentsToMerge.Num())
+		{
+			UWorld* World = ComponentsToMerge[0]->GetWorld();
+			checkf(World != nullptr, TEXT("Invalid World retrieved from Mesh components"));
+			const float ScreenAreaSize = TNumericLimits<float>::Max();
+			MeshUtilities.MergeComponentsToStaticMesh(ComponentsToMerge, World, SettingsObject->Settings, nullptr, nullptr, PackageName, AssetsToSync, MergedActorLocation, ScreenAreaSize, true);
+		}
 	}
 
 	if (AssetsToSync.Num())
diff --git a/Engine/Source/Editor/MergeActors/Private/MeshProxyTool/MeshProxyTool.cpp b/Engine/Source/Editor/MergeActors/Private/MeshProxyTool/MeshProxyTool.cpp
--- a/Engine/Source/Editor/MergeActors/Private/MeshProxyTool/MeshProxyTool.cpp
+++ b/Engine/Source/Editor/MergeActors/Private/MeshProxyTool/MeshProxyTool.cpp
@@ -136,7 +136,7 @@ bool FMeshProxyTool::RunMerge(const FString& PackageName)
 		for (const TSharedPtr<FMergeComponentData>& SelectedComponent : SelectedComponents)
 		{
 			// Determine whether or not this component should be incorporated according the user settings
-			if (SelectedComponent->bShouldIncorporate)
+			if (SelectedComponent->bShouldIncorporate && SelectedComponent->PrimComponent.IsValid())
 			{
 				if (UStaticMeshComponent* StaticMeshComponent = Cast<UStaticMeshComponent>(SelectedComponent->PrimComponent.Get()))
 					StaticMeshComponentsToMerge.Add(StaticMeshComponent);
@@ -144,8 +144,11 @@ bool FMeshProxyTool::RunMerge(const FString& PackageName)
 		}
 		StaticMeshComponentsToMerge.RemoveAll([](UStaticMeshComponent* Val) { return !((Val->GetClass() == UStaticMeshComponent::StaticClass() && Val->GetStaticMesh()) || Val->IsA(USplineMeshComponent::StaticClass())); });
 		
-		FGuid JobGuid = FGuid::NewGuid();
-		MeshMergeUtilities.CreateProxyMesh(StaticMeshComponentsToMerge, SettingsObject->Settings, nullptr, PackageName, JobGuid, ProxyDelegate);
+		if ( StaticMeshComponentsToMerge.Num())
+		{
+			FGuid JobGuid = FGuid::NewGuid();
+			MeshMergeUtilities.CreateProxyMesh(StaticMeshComponentsToMerge, SettingsObject->Settings, nullptr, PackageName, JobGuid, ProxyDelegate);
+		}
 
 		GEditor->EndTransaction();
 		GWarn->EndSlowTask();
