MergeProxyUtils declaration shadows a local variable warnings
#fix Removed unused variables which were causing the warnings

diff --git a/Engine/Source/Editor/MergeActors/Private/MergeProxyUtils/Utils.cpp b/Engine/Source/Editor/MergeActors/Private/MergeProxyUtils/Utils.cpp
--- a/Engine/Source/Editor/MergeActors/Private/MergeProxyUtils/Utils.cpp
+++ b/Engine/Source/Editor/MergeActors/Private/MergeProxyUtils/Utils.cpp
@@ -122,28 +122,11 @@ void FComponentSelectionControl::UpdateSelectedStaticMeshComponents()
 TSharedRef<ITableRow> FComponentSelectionControl::MakeComponentListItemWidget(TSharedPtr<FMergeComponentData> ComponentData, const TSharedRef<STableViewBase>& OwnerTable)
 {
 	check(ComponentData->PrimComponent != nullptr);
-
-	// Retrieve information about the mesh component
-	const FString OwningActorName = ComponentData->PrimComponent->GetOwner()->GetName();
-
+	
 	// If box should be enabled
 	bool bEnabled = true;
 	bool bIsMesh = false;
 
-	FString ComponentInfo;
-	if (UStaticMeshComponent* StaticMeshComponent = Cast<UStaticMeshComponent>(ComponentData->PrimComponent.Get()))
-	{
-		ComponentInfo = (StaticMeshComponent->GetStaticMesh() != nullptr) ? StaticMeshComponent->GetStaticMesh()->GetName() : TEXT("No Static Mesh Available");
-		bEnabled = (StaticMeshComponent->GetStaticMesh() != nullptr);
-		bIsMesh = true;
-	}
-	else if (UShapeComponent* ShapeComponent = Cast<UShapeComponent>(ComponentData->PrimComponent.Get()))
-	{
-		ComponentInfo = ShapeComponent->GetClass()->GetName();
-	}
-
-	const FString ComponentName = ComponentData->PrimComponent->GetName();
-
 	// See if we stored a checkbox state for this mesh component, and set accordingly
 	ECheckBoxState State = bEnabled ? ECheckBoxState::Checked : ECheckBoxState::Unchecked;
 	ECheckBoxState* StoredState = StoredCheckBoxStates.Find(ComponentData->PrimComponent.Get());
@@ -152,7 +135,6 @@ TSharedRef<ITableRow> FComponentSelectionControl::MakeComponentListItemWidget(TS
 		State = *StoredState;
 	}
 
-
 	return SNew(STableRow<TSharedPtr<FMergeComponentData>>, OwnerTable)
 		[
 			SNew(SBox)
