Add a reset to defaults button to Merge Actors window
#fix Ensure we do not use the CDO as the object we inspect in the tool detailsview, rather duplicate the CDO so per-property checks with the default actually works
#misc Clean up objects whenever the MergeActor module shuts down

diff --git a/Engine/Source/Editor/MergeActors/Private/MergeActorsModule.cpp b/Engine/Source/Editor/MergeActors/Private/MergeActorsModule.cpp
--- a/Engine/Source/Editor/MergeActors/Private/MergeActorsModule.cpp
+++ b/Engine/Source/Editor/MergeActors/Private/MergeActorsModule.cpp
@@ -134,8 +134,10 @@ void FMergeActorsModule::ShutdownModule()
 {
 	if (FSlateApplication::IsInitialized())
 	{
-		FGlobalTabmanager::Get()->UnregisterNomadTabSpawner(MergeActorsTabName);
+		FGlobalTabmanager::Get()->UnregisterNomadTabSpawner(MergeActorsTabName);		
 	}
+
+	MergeActorsTools.Empty();
 }
 
 
diff --git a/Engine/Source/Editor/MergeActors/Private/MeshInstancingTool/MeshInstancingTool.cpp b/Engine/Source/Editor/MergeActors/Private/MeshInstancingTool/MeshInstancingTool.cpp
--- a/Engine/Source/Editor/MergeActors/Private/MeshInstancingTool/MeshInstancingTool.cpp
+++ b/Engine/Source/Editor/MergeActors/Private/MeshInstancingTool/MeshInstancingTool.cpp
@@ -21,14 +21,22 @@
 #include "ScopedTransaction.h"
 #include "MeshMergeModule.h"
 
-
 #define LOCTEXT_NAMESPACE "MeshInstancingTool"
 
+bool UMeshInstancingSettingsObject::bInitialized = false;
+UMeshInstancingSettingsObject* UMeshInstancingSettingsObject::DefaultSettings = nullptr;
+
 FMeshInstancingTool::FMeshInstancingTool()
 {
 	SettingsObject = UMeshInstancingSettingsObject::Get();
 }
 
+FMeshInstancingTool::~FMeshInstancingTool()
+{
+	UMeshInstancingSettingsObject::Destroy();
+	SettingsObject = nullptr;
+}
+
 TSharedRef<SWidget> FMeshInstancingTool::GetWidget()
 {
 	SAssignNew(InstancingDialog, SMeshInstancingDialog, this);
@@ -155,5 +163,4 @@ bool FMeshInstancingTool::CanMerge() const
 {	
 	return InstancingDialog->GetNumSelectedMeshComponents() >= 1;
 }
-
-#undef LOCTEXT_NAMESPACE
+#undef LOCTEXT_NAMESPACE
\ No newline at end of file
diff --git a/Engine/Source/Editor/MergeActors/Private/MeshInstancingTool/MeshInstancingTool.h b/Engine/Source/Editor/MergeActors/Private/MeshInstancingTool/MeshInstancingTool.h
--- a/Engine/Source/Editor/MergeActors/Private/MeshInstancingTool/MeshInstancingTool.h
+++ b/Engine/Source/Editor/MergeActors/Private/MeshInstancingTool/MeshInstancingTool.h
@@ -24,9 +24,36 @@ public:
 
 	static UMeshInstancingSettingsObject* Get()
 	{
-		return GetMutableDefault<UMeshInstancingSettingsObject>();
+		if (!bInitialized)
+		{
+			DefaultSettings = DuplicateObject(GetMutableDefault<UMeshInstancingSettingsObject>(), nullptr);
+			DefaultSettings->AddToRoot();
+			bInitialized = true;
+		}
+
+		return DefaultSettings;
+	}
+
+	static void Destroy()
+	{
+		if (bInitialized)
+		{
+			if (DefaultSettings)
+			{
+				DefaultSettings->RemoveFromRoot();
+				DefaultSettings->MarkPendingKill();
+			}
+
+			DefaultSettings = false;
+			bInitialized = false;
+		}
 	}
 
+protected:
+	static bool bInitialized;
+	// This is a singleton, duplicate default object
+	static UMeshInstancingSettingsObject* DefaultSettings;
+
 public:
 	UPROPERTY(EditAnywhere, meta = (ShowOnlyInnerProperties), Category = MergeSettings)
 	FMeshInstancingSettings Settings;
@@ -42,6 +69,7 @@ class FMeshInstancingTool : public IMergeActorsTool
 public:
 
 	FMeshInstancingTool();
+	~FMeshInstancingTool();
 
 	// IMergeActorsTool interface
 	virtual TSharedRef<SWidget> GetWidget() override;
diff --git a/Engine/Source/Editor/MergeActors/Private/MeshInstancingTool/SMeshInstancingDialog.cpp b/Engine/Source/Editor/MergeActors/Private/MeshInstancingTool/SMeshInstancingDialog.cpp
--- a/Engine/Source/Editor/MergeActors/Private/MeshInstancingTool/SMeshInstancingDialog.cpp
+++ b/Engine/Source/Editor/MergeActors/Private/MeshInstancingTool/SMeshInstancingDialog.cpp
@@ -340,6 +340,7 @@ void SMeshInstancingDialog::CreateSettingsView()
 	DetailsViewArgs.bLockable = true;
 	DetailsViewArgs.NameAreaSettings = FDetailsViewArgs::ComponentsAndActorsUseNameArea;
 	DetailsViewArgs.bCustomNameAreaLocation = false;
+	DetailsViewArgs.NameAreaSettings = FDetailsViewArgs::ENameAreaSettings::HideNameArea;
 	DetailsViewArgs.bCustomFilterAreaLocation = true;
 	DetailsViewArgs.DefaultsOnlyVisibility = EEditDefaultsOnlyNodeVisibility::Hide;
 	
diff --git a/Engine/Source/Editor/MergeActors/Private/MeshMergingTool/MeshMergingTool.cpp b/Engine/Source/Editor/MergeActors/Private/MeshMergingTool/MeshMergingTool.cpp
--- a/Engine/Source/Editor/MergeActors/Private/MeshMergingTool/MeshMergingTool.cpp
+++ b/Engine/Source/Editor/MergeActors/Private/MeshMergingTool/MeshMergingTool.cpp
@@ -24,12 +24,21 @@
 
 #define LOCTEXT_NAMESPACE "MeshMergingTool"
 
+UMeshMergingSettingsObject* UMeshMergingSettingsObject::DefaultSettings = nullptr;
+bool UMeshMergingSettingsObject::bInitialized = false;
+
 FMeshMergingTool::FMeshMergingTool()
 	: bReplaceSourceActors(false)
 {
 	SettingsObject = UMeshMergingSettingsObject::Get();
 }
 
+FMeshMergingTool::~FMeshMergingTool()
+{
+	UMeshMergingSettingsObject::Destroy();
+	SettingsObject = nullptr;
+}
+
 TSharedRef<SWidget> FMeshMergingTool::GetWidget()
 {
 	SAssignNew(MergingDialog, SMeshMergingDialog, this);
@@ -172,3 +181,6 @@ bool FMeshMergingTool::CanMerge() const
 }
 
 #undef LOCTEXT_NAMESPACE
+
+
+
diff --git a/Engine/Source/Editor/MergeActors/Private/MeshMergingTool/MeshMergingTool.h b/Engine/Source/Editor/MergeActors/Private/MeshMergingTool/MeshMergingTool.h
--- a/Engine/Source/Editor/MergeActors/Private/MeshMergingTool/MeshMergingTool.h
+++ b/Engine/Source/Editor/MergeActors/Private/MeshMergingTool/MeshMergingTool.h
@@ -25,20 +25,40 @@ public:
 		Settings.LODSelectionType = EMeshLODSelectionType::AllLODs;
 	}
 
-	static UMeshMergingSettingsObject* Get()
-	{
-		static bool bInitialized = false;
-		// This is a singleton, use default object
-		UMeshMergingSettingsObject* DefaultSettings = GetMutableDefault<UMeshMergingSettingsObject>();
+	
 
+	static UMeshMergingSettingsObject* Get()
+	{		
+		// This is a singleton, duplicate default object
 		if (!bInitialized)
 		{
+			DefaultSettings = DuplicateObject(GetMutableDefault<UMeshMergingSettingsObject>(), nullptr);
+			DefaultSettings->AddToRoot();			
 			bInitialized = true;
 		}
 
 		return DefaultSettings;
 	}
 
+	static void Destroy()
+	{
+		if (bInitialized)
+		{
+			if (DefaultSettings)
+			{
+				DefaultSettings->RemoveFromRoot();
+				DefaultSettings->MarkPendingKill();
+			}
+			
+			DefaultSettings = false;
+			bInitialized = false;
+		}
+	}
+
+protected:
+	static UMeshMergingSettingsObject* DefaultSettings;
+	static bool bInitialized;
+
 public:
 	UPROPERTY(editanywhere, meta = (ShowOnlyInnerProperties), Category = MergeSettings)
 	FMeshMergingSettings Settings;
@@ -54,6 +74,7 @@ class FMeshMergingTool : public IMergeActorsTool
 public:
 
 	FMeshMergingTool();
+	~FMeshMergingTool();
 
 	// IMergeActorsTool interface
 	virtual TSharedRef<SWidget> GetWidget() override;
diff --git a/Engine/Source/Editor/MergeActors/Private/MeshMergingTool/SMeshMergingDialog.cpp b/Engine/Source/Editor/MergeActors/Private/MeshMergingTool/SMeshMergingDialog.cpp
--- a/Engine/Source/Editor/MergeActors/Private/MeshMergingTool/SMeshMergingDialog.cpp
+++ b/Engine/Source/Editor/MergeActors/Private/MeshMergingTool/SMeshMergingDialog.cpp
@@ -210,6 +210,7 @@ void SMeshMergingDialog::CreateSettingsView()
 	DetailsViewArgs.bLockable = true;
 	DetailsViewArgs.NameAreaSettings = FDetailsViewArgs::ComponentsAndActorsUseNameArea;
 	DetailsViewArgs.bCustomNameAreaLocation = false;
+	DetailsViewArgs.NameAreaSettings = FDetailsViewArgs::ENameAreaSettings::HideNameArea;
 	DetailsViewArgs.bCustomFilterAreaLocation = true;
 	DetailsViewArgs.DefaultsOnlyVisibility = EEditDefaultsOnlyNodeVisibility::Hide;
 	
diff --git a/Engine/Source/Editor/MergeActors/Private/MeshProxyTool/MeshProxyTool.cpp b/Engine/Source/Editor/MergeActors/Private/MeshProxyTool/MeshProxyTool.cpp
--- a/Engine/Source/Editor/MergeActors/Private/MeshProxyTool/MeshProxyTool.cpp
+++ b/Engine/Source/Editor/MergeActors/Private/MeshProxyTool/MeshProxyTool.cpp
@@ -23,12 +23,20 @@
 
 #define LOCTEXT_NAMESPACE "MeshProxyTool"
 
+bool UMeshProxySettingsObject::bInitialized = false;
+UMeshProxySettingsObject* UMeshProxySettingsObject::DefaultSettings = nullptr;
 
 FMeshProxyTool::FMeshProxyTool()
 {
 	SettingsObject = UMeshProxySettingsObject::Get();
 }
 
+FMeshProxyTool::~FMeshProxyTool()
+{
+	UMeshProxySettingsObject::Destroy();
+	SettingsObject = nullptr;
+}
+
 TSharedRef<SWidget> FMeshProxyTool::GetWidget()
 {
 	SAssignNew(ProxyDialog, SMeshProxyDialog, this);
@@ -278,3 +286,4 @@ TSharedRef<SWidget> FThirdPartyMeshProxyTool::GetWidget()
 }
 
 #undef LOCTEXT_NAMESPACE
+
diff --git a/Engine/Source/Editor/MergeActors/Private/MeshProxyTool/MeshProxyTool.h b/Engine/Source/Editor/MergeActors/Private/MeshProxyTool/MeshProxyTool.h
--- a/Engine/Source/Editor/MergeActors/Private/MeshProxyTool/MeshProxyTool.h
+++ b/Engine/Source/Editor/MergeActors/Private/MeshProxyTool/MeshProxyTool.h
@@ -20,19 +20,36 @@ public:
 	}
 
 	static UMeshProxySettingsObject* Get()
-	{
-		static bool bInitialized = false;
-		// This is a singleton, use default object
-		UMeshProxySettingsObject* DefaultSettings = GetMutableDefault<UMeshProxySettingsObject>();
-
+	{	
+		// This is a singleton, duplicate default object
 		if (!bInitialized)
 		{
+			DefaultSettings = DuplicateObject(GetMutableDefault<UMeshProxySettingsObject>(), nullptr);
+			DefaultSettings->AddToRoot();
 			bInitialized = true;
 		}
 
 		return DefaultSettings;
 	}
 
+	static void Destroy()
+	{
+		if (bInitialized)
+		{
+			if (DefaultSettings)
+			{
+				DefaultSettings->RemoveFromRoot();
+				DefaultSettings->MarkPendingKill();
+			}
+
+			DefaultSettings = false;
+			bInitialized = false;
+		}
+	}
+
+protected:
+	static bool bInitialized;
+	static UMeshProxySettingsObject* DefaultSettings;
 public:
 	UPROPERTY(editanywhere, meta = (ShowOnlyInnerProperties), Category = ProxySettings)
 	FMeshProxySettings  Settings;
@@ -47,6 +64,7 @@ class FMeshProxyTool : public IMergeActorsTool
 
 public:
 	FMeshProxyTool();
+	~FMeshProxyTool();
 
 	// IMergeActorsTool interface
 	virtual TSharedRef<SWidget> GetWidget() override;
diff --git a/Engine/Source/Editor/MergeActors/Private/MeshProxyTool/SMeshProxyDialog.cpp b/Engine/Source/Editor/MergeActors/Private/MeshProxyTool/SMeshProxyDialog.cpp
--- a/Engine/Source/Editor/MergeActors/Private/MeshProxyTool/SMeshProxyDialog.cpp
+++ b/Engine/Source/Editor/MergeActors/Private/MeshProxyTool/SMeshProxyDialog.cpp
@@ -191,6 +191,7 @@ void SMeshProxyDialog::CreateSettingsView()
 	DetailsViewArgs.bCustomNameAreaLocation = false;
 	DetailsViewArgs.bCustomFilterAreaLocation = true;
 	DetailsViewArgs.DefaultsOnlyVisibility = EEditDefaultsOnlyNodeVisibility::Hide;
+	DetailsViewArgs.NameAreaSettings = FDetailsViewArgs::ENameAreaSettings::HideNameArea;
 
 
 	// Tiny hack to hide this setting, since we have no way / value to go off to 
