Fix Mac build errors, amazed MSVC did not complain about this

diff --git a/Engine/Source/Editor/MergeActors/Private/MeshInstancingTool/MeshInstancingTool.h b/Engine/Source/Editor/MergeActors/Private/MeshInstancingTool/MeshInstancingTool.h
--- a/Engine/Source/Editor/MergeActors/Private/MeshInstancingTool/MeshInstancingTool.h
+++ b/Engine/Source/Editor/MergeActors/Private/MeshInstancingTool/MeshInstancingTool.h
@@ -44,7 +44,7 @@ public:
 				DefaultSettings->MarkPendingKill();
 			}
 
-			DefaultSettings = false;
+			DefaultSettings = nullptr;
 			bInitialized = false;
 		}
 	}
diff --git a/Engine/Source/Editor/MergeActors/Private/MeshMergingTool/MeshMergingTool.h b/Engine/Source/Editor/MergeActors/Private/MeshMergingTool/MeshMergingTool.h
--- a/Engine/Source/Editor/MergeActors/Private/MeshMergingTool/MeshMergingTool.h
+++ b/Engine/Source/Editor/MergeActors/Private/MeshMergingTool/MeshMergingTool.h
@@ -50,7 +50,7 @@ public:
 				DefaultSettings->MarkPendingKill();
 			}
 			
-			DefaultSettings = false;
+			DefaultSettings = nullptr;
 			bInitialized = false;
 		}
 	}
diff --git a/Engine/Source/Editor/MergeActors/Private/MeshProxyTool/MeshProxyTool.h b/Engine/Source/Editor/MergeActors/Private/MeshProxyTool/MeshProxyTool.h
--- a/Engine/Source/Editor/MergeActors/Private/MeshProxyTool/MeshProxyTool.h
+++ b/Engine/Source/Editor/MergeActors/Private/MeshProxyTool/MeshProxyTool.h
@@ -42,7 +42,7 @@ public:
 				DefaultSettings->MarkPendingKill();
 			}
 
-			DefaultSettings = false;
+			DefaultSettings = nullptr;
 			bInitialized = false;
 		}
 	}
