Fix implicit conversion from pointer to bool

diff --git a/Engine/Source/Editor/MergeActors/Private/MeshProxyTool/MeshProxyTool.cpp b/Engine/Source/Editor/MergeActors/Private/MeshProxyTool/MeshProxyTool.cpp
--- a/Engine/Source/Editor/MergeActors/Private/MeshProxyTool/MeshProxyTool.cpp
+++ b/Engine/Source/Editor/MergeActors/Private/MeshProxyTool/MeshProxyTool.cpp
@@ -154,7 +154,7 @@ bool FMeshProxyTool::RunMerge(const FString& PackageName)
 		StaticMeshComponentsToMerge.RemoveAll([](UStaticMeshComponent* Val)
 		{
 			bool bValidClass = Val->GetClass() == UStaticMeshComponent::StaticClass() || Val->IsA<UInstancedStaticMeshComponent>() || Val->IsA<USplineMeshComponent>();
-			bool bValidMesh = Val->GetStaticMesh();
+			bool bValidMesh = Val->GetStaticMesh() != nullptr;
 			return !bValidClass || !bValidMesh;
 		});
 		
