Merge Actors tool overwrites existing files without warning user of data loss
#fix Set static mesh as relevant class for override file dialog

diff --git a/Engine/Source/Editor/MergeActors/Private/SMergeActorsToolbar.cpp b/Engine/Source/Editor/MergeActors/Private/SMergeActorsToolbar.cpp
--- a/Engine/Source/Editor/MergeActors/Private/SMergeActorsToolbar.cpp
+++ b/Engine/Source/Editor/MergeActors/Private/SMergeActorsToolbar.cpp
@@ -138,6 +138,7 @@ FReply SMergeActorsToolbar::OnMergeActorsClicked()
 			SaveAssetDialogConfig.DefaultPath = DefaultPath;
 			SaveAssetDialogConfig.DefaultAssetName = DefaultName;
 			SaveAssetDialogConfig.ExistingAssetPolicy = ESaveAssetDialogExistingAssetPolicy::AllowButWarn;
+			SaveAssetDialogConfig.AssetClassNames = { UStaticMesh::StaticClass()->GetFName() };
 
 			FContentBrowserModule& ContentBrowserModule = FModuleManager::LoadModuleChecked<FContentBrowserModule>("ContentBrowser");
 			FString SaveObjectPath = ContentBrowserModule.Get().CreateModalSaveAssetDialog(SaveAssetDialogConfig);
