Merged Actors incorrectly renamed when Merge Materials is enabled
#fix correctly name resulting actor when actor merging is set to replace

diff --git a/Engine/Source/Editor/MergeActors/Private/MeshMergingTool/MeshMergingTool.cpp b/Engine/Source/Editor/MergeActors/Private/MeshMergingTool/MeshMergingTool.cpp
--- a/Engine/Source/Editor/MergeActors/Private/MeshMergingTool/MeshMergingTool.cpp
+++ b/Engine/Source/Editor/MergeActors/Private/MeshMergingTool/MeshMergingTool.cpp
@@ -159,7 +159,7 @@ bool FMeshMergingTool::RunMerge(const FString& PackageName)
 								
 				AStaticMeshActor* MergedActor = World->SpawnActor<AStaticMeshActor>(MergedActorLocation, MergedActorRotation, Params);
 				MergedActor->GetStaticMeshComponent()->SetStaticMesh(MergedMesh);
-				MergedActor->SetActorLabel(AssetsToSync[0]->GetName());
+				MergedActor->SetActorLabel(MergedMesh->GetName());
 				World->UpdateCullDistanceVolumes(MergedActor, MergedActor->GetStaticMeshComponent());
 				// Remove source actors
 				for (AActor* Actor : Actors)
