Fix MergeActor Crash (update selection)

diff --git a/Engine/Source/Editor/MergeActors/Private/MeshMergingTool/MeshMergingTool.cpp b/Engine/Source/Editor/MergeActors/Private/MeshMergingTool/MeshMergingTool.cpp
--- a/Engine/Source/Editor/MergeActors/Private/MeshMergingTool/MeshMergingTool.cpp
+++ b/Engine/Source/Editor/MergeActors/Private/MeshMergingTool/MeshMergingTool.cpp
@@ -173,6 +173,8 @@ bool FMeshMergingTool::RunMerge(const FString& PackageName)
 				MergedActor->GetStaticMeshComponent()->SetStaticMesh(MergedMesh);
 				MergedActor->SetActorLabel(MergedMesh->GetName());
 				World->UpdateCullDistanceVolumes(MergedActor, MergedActor->GetStaticMeshComponent());
+				GEditor->SelectNone(true, true);
+				GEditor->SelectActor(MergedActor, true, true);
 				// Remove source actors
 				for (AActor* Actor : Actors)
 				{
