- Adapted the MergeActors dialog to the new UX style.
- Refactored the MergeActors different panels to reuse the same code.
- Adapted quick merge actions (required some refactoring)
- Fixed some issues with the "replace source actors" feature vs. actual selection
- Fixed crash when using the instancing mode with "replace source actors" enabled.

diff --git a/Engine/Source/Developer/MeshMergeUtilities/Private/MeshMergeUtilities.cpp b/Engine/Source/Developer/MeshMergeUtilities/Private/MeshMergeUtilities.cpp
--- a/Engine/Source/Developer/MeshMergeUtilities/Private/MeshMergeUtilities.cpp
+++ b/Engine/Source/Developer/MeshMergeUtilities/Private/MeshMergeUtilities.cpp
@@ -3093,7 +3093,7 @@ void FMeshMergeUtilities::CreateMergedRawMeshes(FMeshMergeDataTracker& InDataTra
 	}
 }
 
-void FMeshMergeUtilities::MergeComponentsToInstances(const TArray<UPrimitiveComponent*>& ComponentsToMerge, UWorld* World, ULevel* Level, const FMeshInstancingSettings& InSettings, bool bActuallyMerge /*= true*/, FText* OutResultsText /*= nullptr*/) const
+void FMeshMergeUtilities::MergeComponentsToInstances(const TArray<UPrimitiveComponent*>& ComponentsToMerge, UWorld* World, ULevel* Level, const FMeshInstancingSettings& InSettings, bool bActuallyMerge /*= true*/, bool bReplaceSourceActors /* = false */, FText* OutResultsText /*= nullptr*/) const
 {
 	auto HasInstanceVertexColors = [](UStaticMeshComponent* StaticMeshComponent)
 	{
@@ -3329,11 +3329,11 @@ void FMeshMergeUtilities::MergeComponentsToInstances(const TArray<UPrimitiveComp
 				// Now clean up our original actors
 				for(AActor* ActorToCleanUp : ActorsToCleanUp)
 				{
-					if(InSettings.MeshReplacementMethod == EMeshInstancingReplacementMethod::RemoveOriginalActors)
+					if (bReplaceSourceActors)
 					{
 						ActorToCleanUp->Destroy();
 					}
-					else if(InSettings.MeshReplacementMethod == EMeshInstancingReplacementMethod::KeepOriginalActorsAsEditorOnly)
+					else
 					{
 						ActorToCleanUp->Modify();
 						ActorToCleanUp->bIsEditorOnlyActor = true;
@@ -3358,12 +3358,21 @@ void FMeshMergeUtilities::MergeComponentsToInstances(const TArray<UPrimitiveComp
 					GEditor->GetSelectedActors()->EndBatchSelectOperation();
 				};
 
-				FNotificationInfo NotificationInfo(FText::Format(LOCTEXT("CreatedInstancedActorsMessage", "Created {0} Instanced Actor(s)"), FText::AsNumber(ActorEntries.Num())));
-				NotificationInfo.Hyperlink = FSimpleDelegate::CreateLambda(SelectActorsLambda);
-				NotificationInfo.HyperlinkText = LOCTEXT("SelectActorsHyperlink", "Select Actors");
-				NotificationInfo.ExpireDuration = 5.0f;
+				// Always change selection if we removed the source actors,
+				// Otherwise, allow selection change through notification
+				if (bReplaceSourceActors)
+				{
+					SelectActorsLambda();
+				}
+				else
+				{
+					FNotificationInfo NotificationInfo(FText::Format(LOCTEXT("CreatedInstancedActorsMessage", "Created {0} Instanced Actor(s)"), FText::AsNumber(ActorEntries.Num())));
+					NotificationInfo.Hyperlink = FSimpleDelegate::CreateLambda(SelectActorsLambda);
+					NotificationInfo.HyperlinkText = LOCTEXT("SelectActorsHyperlink", "Select Actors");
+					NotificationInfo.ExpireDuration = 5.0f;
 
-				FSlateNotificationManager::Get().AddNotification(NotificationInfo);
+					FSlateNotificationManager::Get().AddNotification(NotificationInfo);
+				}
 			}
 		}
 	}
diff --git a/Engine/Source/Developer/MeshMergeUtilities/Private/MeshMergeUtilities.h b/Engine/Source/Developer/MeshMergeUtilities/Private/MeshMergeUtilities.h
--- a/Engine/Source/Developer/MeshMergeUtilities/Private/MeshMergeUtilities.h
+++ b/Engine/Source/Developer/MeshMergeUtilities/Private/MeshMergeUtilities.h
@@ -24,7 +24,7 @@ public:
 	virtual void BakeMaterialsForComponent(USkeletalMeshComponent* SkeletalMeshComponent) const override;
 	virtual void BakeMaterialsForMesh(UStaticMesh* Mesh) const override;
 	virtual void MergeComponentsToStaticMesh(const TArray<UPrimitiveComponent*>& ComponentsToMerge, UWorld* World, const FMeshMergingSettings& InSettings, UMaterialInterface* InBaseMaterial, UPackage* InOuter, const FString& InBasePackageName, TArray<UObject*>& OutAssetsToSync, FVector& OutMergedActorLocation, const float ScreenSize, bool bSilent /*= false*/) const;
-	virtual void MergeComponentsToInstances(const TArray<UPrimitiveComponent*>& ComponentsToMerge, UWorld* World, ULevel* Level, const FMeshInstancingSettings& InSettings, bool bActuallyMerge = true, FText* OutResultsText = nullptr) const override;
+	virtual void MergeComponentsToInstances(const TArray<UPrimitiveComponent*>& ComponentsToMerge, UWorld* World, ULevel* Level, const FMeshInstancingSettings& InSettings, bool bActuallyMerge = true, bool bReplaceSourceActors = false, FText* OutResultsText = nullptr) const override;
 
 	virtual void CreateProxyMesh(const TArray<AActor*>& InActors, const struct FMeshProxySettings& InMeshProxySettings, UMaterialInterface* InBaseMaterial, UPackage* InOuter, const FString& InProxyBasePackageName, const FGuid InGuid, const FCreateProxyDelegate& InProxyCreatedDelegate, const bool bAllowAsync = false, const float ScreenSize = 1.0f) const override;
 	virtual void CreateProxyMesh(const TArray<UStaticMeshComponent*>& InStaticMeshComps, const struct FMeshProxySettings& InMeshProxySettings, UMaterialInterface* InBaseMaterial, UPackage* InOuter, const FString& InProxyBasePackageName, const FGuid InGuid, const FCreateProxyDelegate& InProxyCreatedDelegate, const bool bAllowAsync = false, const float ScreenSize = 1.0f) const override;
diff --git a/Engine/Source/Developer/MeshMergeUtilities/Public/IMeshMergeUtilities.h b/Engine/Source/Developer/MeshMergeUtilities/Public/IMeshMergeUtilities.h
--- a/Engine/Source/Developer/MeshMergeUtilities/Public/IMeshMergeUtilities.h
+++ b/Engine/Source/Developer/MeshMergeUtilities/Public/IMeshMergeUtilities.h
@@ -77,7 +77,7 @@ public:
 	* @param OutResultsText (optional) Results of the merge
 	* @return void
 	*/
-	virtual void MergeComponentsToInstances(const TArray<UPrimitiveComponent*>& ComponentsToMerge, UWorld* World, ULevel* Level, const FMeshInstancingSettings& InSettings, bool bActuallyMerge = true, FText* OutResultsText = nullptr) const = 0;
+	virtual void MergeComponentsToInstances(const TArray<UPrimitiveComponent*>& ComponentsToMerge, UWorld* World, ULevel* Level, const FMeshInstancingSettings& InSettings, bool bActuallyMerge = true, bool bReplaceSourceActors = false, FText* OutResultsText = nullptr) const = 0;
 
 	/**
 	*	Merges list of actors into single proxy mesh
diff --git a/Engine/Source/Editor/DetailCustomizations/Private/MeshProxySettingsCustomizations.cpp b/Engine/Source/Editor/DetailCustomizations/Private/MeshProxySettingsCustomizations.cpp
--- a/Engine/Source/Editor/DetailCustomizations/Private/MeshProxySettingsCustomizations.cpp
+++ b/Engine/Source/Editor/DetailCustomizations/Private/MeshProxySettingsCustomizations.cpp
@@ -48,8 +48,26 @@ void FMeshProxySettingsCustomizations::CustomizeChildren(TSharedRef<IPropertyHan
 	IMeshReductionManagerModule& ModuleManager = FModuleManager::Get().LoadModuleChecked<IMeshReductionManagerModule>("MeshReductionInterface");
 	IMeshMerging* MergeModule = ModuleManager.GetMeshMergingInterface();
 
-	IDetailGroup& MeshSettingsGroup = ChildBuilder.AddGroup(NAME_None, FText::FromString("Proxy Settings"));
+	// Respect the ShowInnerProperties property from the struct property.
+	static const FName ShowOnlyInners("ShowOnlyInnerProperties");
+	bool bCreateSettingsGroup = !StructPropertyHandle->HasMetaData(ShowOnlyInners);
+	IDetailGroup* MeshSettingsGroup = nullptr;
 
+	if (bCreateSettingsGroup)
+	{
+		MeshSettingsGroup = &ChildBuilder.AddGroup(NAME_None, FText::FromString("Proxy Settings"));
+	}
+
+	auto AddPropertyToGroup = [bCreateSettingsGroup, &ChildBuilder, &MeshSettingsGroup](const TSharedRef<IPropertyHandle>& prop) -> IDetailPropertyRow& {
+		if (bCreateSettingsGroup)
+		{
+			return MeshSettingsGroup->AddPropertyRow(prop);
+		}
+		else
+		{
+			return ChildBuilder.AddProperty(prop);
+		}
+	};
 
 
 	TSharedPtr< IPropertyHandle > HardAngleThresholdPropertyHandle        = PropertyHandles.FindChecked(GET_MEMBER_NAME_CHECKED(FMeshProxySettings, HardAngleThreshold));
@@ -69,55 +87,55 @@ void FMeshProxySettingsCustomizations::CustomizeChildren(TSharedRef<IPropertyHan
 		// Handle special property cases (done inside the loop to maintain order according to the struct
 		if (Iter.Value() == HardAngleThresholdPropertyHandle)
 		{
-			IDetailPropertyRow& MeshProxySettingsRow = MeshSettingsGroup.AddPropertyRow(Iter.Value().ToSharedRef());
+			IDetailPropertyRow& MeshProxySettingsRow = AddPropertyToGroup(Iter.Value().ToSharedRef());
 			MeshProxySettingsRow.ToolTip(FText::FromString(FString("Angle at which a hard edge is introduced between faces.  Note: Increases vertex count and may introduce additional UV seams.  It is only recommended if not using normals maps")));
 			MeshProxySettingsRow.Visibility(TAttribute<EVisibility>(this, &FMeshProxySettingsCustomizations::IsHardAngleThresholdVisible));
 		}
 		else if (Iter.Value() == NormalCalcMethodPropertyHandle)
 		{
-			IDetailPropertyRow& MeshProxySettingsRow = MeshSettingsGroup.AddPropertyRow(Iter.Value().ToSharedRef());
+			IDetailPropertyRow& MeshProxySettingsRow = AddPropertyToGroup(Iter.Value().ToSharedRef());
 			MeshProxySettingsRow.Visibility(TAttribute<EVisibility>(this, &FMeshProxySettingsCustomizations::IsNormalCalcMethodVisible));
 		}
 		else if (Iter.Value() == MaxRayCastDistdPropertyHandle)
 		{
-			IDetailPropertyRow& MeshProxySettingsRow = MeshSettingsGroup.AddPropertyRow(Iter.Value().ToSharedRef());
+			IDetailPropertyRow& MeshProxySettingsRow = AddPropertyToGroup(Iter.Value().ToSharedRef());
 			MeshProxySettingsRow.Visibility(TAttribute<EVisibility>(this, &FMeshProxySettingsCustomizations::IsSearchDistanceVisible));
 		}
 		else if (Iter.Value() == RecalculateNormalsPropertyHandle)
 		{
-			IDetailPropertyRow& MeshProxySettingsRow = MeshSettingsGroup.AddPropertyRow(Iter.Value().ToSharedRef());
+			IDetailPropertyRow& MeshProxySettingsRow = AddPropertyToGroup(Iter.Value().ToSharedRef());
 			MeshProxySettingsRow.Visibility(TAttribute<EVisibility>(this, &FMeshProxySettingsCustomizations::IsRecalculateNormalsVisible));
 		}
 		else if (Iter.Value() == UseLandscapeCullingPropertyHandle)
 		{
-			IDetailPropertyRow& MeshProxySettingsRow = MeshSettingsGroup.AddPropertyRow(Iter.Value().ToSharedRef());
+			IDetailPropertyRow& MeshProxySettingsRow = AddPropertyToGroup(Iter.Value().ToSharedRef());
 			MeshProxySettingsRow.DisplayName(FText::FromString(FString("Enable Volume Culling")));
 			MeshProxySettingsRow.ToolTip(FText::FromString(FString("Allow culling volumes to exclude geometry.")));
 			MeshProxySettingsRow.Visibility(TAttribute<EVisibility>(this, &FMeshProxySettingsCustomizations::IsUseLandscapeCullingVisible));
 		}
 		else if (Iter.Value() == LandscapeCullingPrecisionPropertyHandle)
 		{
-			IDetailPropertyRow& MeshProxySettingsRow = MeshSettingsGroup.AddPropertyRow(Iter.Value().ToSharedRef());
+			IDetailPropertyRow& MeshProxySettingsRow = AddPropertyToGroup(Iter.Value().ToSharedRef());
 			MeshProxySettingsRow.Visibility(TAttribute<EVisibility>(this, &FMeshProxySettingsCustomizations::IsUseLandscapeCullingPrecisionVisible));
 		}
 		else if (Iter.Value() == MergeDistanceHandle)
 		{
-			IDetailPropertyRow& MeshProxySettingsRow = MeshSettingsGroup.AddPropertyRow(Iter.Value().ToSharedRef());
+			IDetailPropertyRow& MeshProxySettingsRow = AddPropertyToGroup(Iter.Value().ToSharedRef());
 			MeshProxySettingsRow.Visibility(TAttribute<EVisibility>(this, &FMeshProxySettingsCustomizations::IsMergeDistanceVisible));
 		}
 		else if (Iter.Value() == UnresolvedGeometryColorHandle)
 		{
-			IDetailPropertyRow& MeshProxySettingsRow = MeshSettingsGroup.AddPropertyRow(Iter.Value().ToSharedRef());
+			IDetailPropertyRow& MeshProxySettingsRow = AddPropertyToGroup(Iter.Value().ToSharedRef());
 			MeshProxySettingsRow.Visibility(TAttribute<EVisibility>(this, &FMeshProxySettingsCustomizations::IsUnresolvedGeometryColorVisible));
 		}
 		else if (Iter.Value() == VoxelSizeHandle)
 		{
-			IDetailPropertyRow& MeshProxySettingsRow = MeshSettingsGroup.AddPropertyRow(Iter.Value().ToSharedRef());
+			IDetailPropertyRow& MeshProxySettingsRow = AddPropertyToGroup(Iter.Value().ToSharedRef());
 			MeshProxySettingsRow.Visibility(TAttribute<EVisibility>(this, &FMeshProxySettingsCustomizations::IsVoxelSizeVisible));
 		}
 		else
 		{
-			IDetailPropertyRow& SettingsRow = MeshSettingsGroup.AddPropertyRow(Iter.Value().ToSharedRef());
+			IDetailPropertyRow& SettingsRow = AddPropertyToGroup(Iter.Value().ToSharedRef());
 		}
 	}
 }
diff --git a/Engine/Source/Editor/EditorStyle/Private/StarshipStyle.cpp b/Engine/Source/Editor/EditorStyle/Private/StarshipStyle.cpp
--- a/Engine/Source/Editor/EditorStyle/Private/StarshipStyle.cpp
+++ b/Engine/Source/Editor/EditorStyle/Private/StarshipStyle.cpp
@@ -4496,10 +4496,10 @@ void FStarshipEditorStyle::FStyle::SetupLevelEditorStyle()
 		Set( "LevelEditor.ToggleVR", new IMAGE_BRUSH( "Icons/VREditor/VR_Editor_Toolbar_Icon", Icon40x40 ) );
 		Set( "LevelEditor.ToggleVR.Small", new IMAGE_BRUSH( "Icons/VREditor/VR_Editor_Toolbar_Icon_Small", Icon20x20 ) );
 
-		Set( "MergeActors.MeshMergingTool", new IMAGE_BRUSH( "Icons/icon_MergeActors_MeshMerging_40x", Icon40x40 ) );
-		Set( "MergeActors.MeshProxyTool", new IMAGE_BRUSH( "Icons/icon_MergeActors_MeshProxy_40x", Icon40x40 ) );
-		Set( "MergeActors.MeshInstancingTool", new IMAGE_BRUSH( "Icons/icon_MergeActors_MeshInstancing_40x", Icon40x40 ) );
-		Set( "MergeActors.TabIcon", new IMAGE_BRUSH("Icons/Icon_MergeActors_MeshMerging_16x", Icon16x16));
+		Set( "MergeActors.MeshMergingTool", new IMAGE_BRUSH_SVG( "Starship/MergeActors/MeshMerging_16", Icon16x16 ) );
+		Set( "MergeActors.MeshProxyTool", new IMAGE_BRUSH_SVG( "Starship/MergeActors/MeshProxy_16", Icon16x16 ) );
+		Set( "MergeActors.MeshInstancingTool", new IMAGE_BRUSH_SVG("Starship/AssetIcons/StaticMeshActor_16", Icon16x16 ) );
+		Set( "MergeActors.TabIcon", new IMAGE_BRUSH_SVG("Starship/MergeActors/MeshMerging_16", Icon16x16));
 		
 		Set( "PlacementBrowser.OptionsMenu", new IMAGE_BRUSH( "Icons/icon_Blueprint_Macro_16x", Icon16x16 ) );
 
diff --git a/Engine/Source/Editor/LevelEditor/LevelEditor.Build.cs b/Engine/Source/Editor/LevelEditor/LevelEditor.Build.cs
--- a/Engine/Source/Editor/LevelEditor/LevelEditor.Build.cs
+++ b/Engine/Source/Editor/LevelEditor/LevelEditor.Build.cs
@@ -18,6 +18,7 @@ public class LevelEditor : ModuleRules
                 "PortalServices",
                 "Persona",
 				"DataLayerEditor",
+				"MergeActors",
 			}
 		);
 
@@ -106,7 +107,8 @@ public class LevelEditor : ModuleRules
 				"VREditor",
                 "Persona",
 				"LevelAssetEditor",
-            }
+				"MergeActors",
+			}
 		);
 
 		if(Target.bWithLiveCoding)
diff --git a/Engine/Source/Editor/LevelEditor/Private/LevelEditorContextMenu.cpp b/Engine/Source/Editor/LevelEditor/Private/LevelEditorContextMenu.cpp
--- a/Engine/Source/Editor/LevelEditor/Private/LevelEditorContextMenu.cpp
+++ b/Engine/Source/Editor/LevelEditor/Private/LevelEditorContextMenu.cpp
@@ -49,6 +49,8 @@
 #include "Framework/Commands/GenericCommands.h"
 #include "LevelViewportActions.h"
 #include "ActorGroupingUtils.h"
+#include "IMergeActorsModule.h"
+#include "IMergeActorsTool.h"
 
 #define LOCTEXT_NAMESPACE "LevelViewportContextMenu"
 
@@ -125,6 +127,13 @@ public:
 	 */
 	static void FillEditMenu(UToolMenu* Menu);
 
+	/**
+	 * Fills in menu options for the merge actors menu
+	 * 
+	 * @param MenuBuilder	The menu to add items to
+	 */
+	static void FillMergeActorsMenu(UToolMenu* Menu);
+
 private:
 	/**
 	 * Fills in menu options for the matinee selection menu
@@ -308,6 +317,17 @@ void FLevelEditorContextMenu::RegisterActorContextMenu()
 			}
 		}
 
+		if (LevelEditorContext->ContextType == ELevelEditorMenuContext::Viewport || LevelEditorContext->ContextType == ELevelEditorMenuContext::SceneOutliner)
+		{
+			FToolMenuSection& Section = InMenu->AddSection("OpenMergeActor");
+
+			Section.AddSubMenu(
+				"MergeActorsSubMenu",
+				FText::Format(LOCTEXT("MergeActorsSubMenu", "Merge Actors ({0})"), FText::AsNumber(SelectedActors.Num())),
+				LOCTEXT("MergeActorsSubMenu_ToolTip", "Merge actors utils"),
+				FNewToolMenuDelegate::CreateStatic(&FLevelEditorContextMenuImpl::FillMergeActorsMenu));
+		}
+
 		// Go to C++ Code
 		if (SelectionInfo.SelectionClass != NULL)
 		{
@@ -375,11 +395,6 @@ void FLevelEditorContextMenu::RegisterActorContextMenu()
 		if (LevelEditorContext->ContextType == ELevelEditorMenuContext::Viewport || LevelEditorContext->ContextType == ELevelEditorMenuContext::SceneOutliner)
 		{
 			LevelEditorCreateActorMenu::FillAddReplaceContextMenuSections(InMenu, LevelEditorContext);
-
-			FToolMenuSection& Section = InMenu->AddSection("OpenMergeActor");
-			Section.AddMenuEntry(FLevelEditorCommands::Get().OpenMergeActor,
-					LOCTEXT("OpenMergeActor", "Merge Actors"),
-					LOCTEXT("OpenMergeActor_ToolTip", "Click to open the Merge Actor panel"));
 		}
 
 		if (GEditor->PlayWorld != NULL)
@@ -1292,6 +1307,40 @@ void FLevelEditorContextMenuImpl::FillEditMenu( UToolMenu* Menu )
 	Section.AddMenuEntry( FGenericCommands::Get().Rename );
 }
 
+void FLevelEditorContextMenuImpl::FillMergeActorsMenu(UToolMenu* Menu)
+{
+	FToolMenuSection& Section = Menu->AddSection("OpenPanel");
+
+	IMergeActorsModule& MergeActorsModule = IMergeActorsModule::Get();
+	TArray<IMergeActorsTool*> MergeActorTools;
+	MergeActorsModule.Get().GetRegisteredMergeActorsTools(MergeActorTools);
+
+	for (IMergeActorsTool* Tool : MergeActorTools)
+	{
+		FUIAction MergeActorToolAction(
+			FExecuteAction::CreateLambda([Tool]() { Tool->RunMergeFromSelection(); }),
+			FCanExecuteAction::CreateLambda([Tool]() { return Tool->CanMergeFromSelection(); }));
+
+		Section.AddMenuEntry(
+			NAME_None,
+			Tool->GetTooltipText(),
+			Tool->GetTooltipText(),
+			FSlateIcon(FEditorStyle::GetStyleSetName(), Tool->GetIconName()),
+			MergeActorToolAction);
+	}
+
+	if (MergeActorTools.Num() > 0)
+	{
+		Section.AddSeparator(NAME_None);
+	}
+
+	Section.AddMenuEntry(
+		FLevelEditorCommands::Get().OpenMergeActor,
+		LOCTEXT("OpenMergeActor", "Merge Actors Settings"),
+		LOCTEXT("OpenMergeActor_ToolTip", "Click to open the Merge Actor panel"),
+		FSlateIcon(FEditorStyle::GetStyleSetName(), "LevelEditor.GameSettings"));
+}
+
 void FLevelScriptEventMenuHelper::FillLevelBlueprintEventsMenu(UToolMenu* Menu, const TArray<AActor*>& SelectedActors)
 {
 	AActor* SelectedActor = (1 == SelectedActors.Num()) ? SelectedActors[0] : NULL;
diff --git a/Engine/Source/Editor/MergeActors/MergeActors.Build.cs b/Engine/Source/Editor/MergeActors/MergeActors.Build.cs
--- a/Engine/Source/Editor/MergeActors/MergeActors.Build.cs
+++ b/Engine/Source/Editor/MergeActors/MergeActors.Build.cs
@@ -19,7 +19,6 @@ public class MergeActors : ModuleRules
 				"AssetRegistry",
                 "ContentBrowser",
                 "Documentation",
-                "LevelEditor",
                 "MeshUtilities",
                 "PropertyEditor",
                 "RawMesh",
@@ -48,7 +47,6 @@ public class MergeActors : ModuleRules
 				"AssetRegistry",
                 "ContentBrowser",
                 "Documentation",
-                "LevelEditor",
                 "MeshUtilities",
                 "MeshMergeUtilities",
                 "MeshReductionInterface",
diff --git a/Engine/Source/Editor/MergeActors/Private/MergeActorsModule.cpp b/Engine/Source/Editor/MergeActors/Private/MergeActorsModule.cpp
--- a/Engine/Source/Editor/MergeActors/Private/MergeActorsModule.cpp
+++ b/Engine/Source/Editor/MergeActors/Private/MergeActorsModule.cpp
@@ -51,6 +51,10 @@ public:
 	 */
 	virtual bool UnregisterMergeActorsTool(IMergeActorsTool* Tool) override;
 
+	/**
+	 * Get list of registered tools
+	 */
+	virtual void GetRegisteredMergeActorsTools(TArray<IMergeActorsTool*>& OutTools) override;
 
 private:
 
@@ -179,5 +183,14 @@ bool FMergeActorsModule::UnregisterMergeActorsTool(IMergeActorsTool* Tool)
 	return false;
 }
 
+void FMergeActorsModule::GetRegisteredMergeActorsTools(TArray<IMergeActorsTool*>& OutTools)
+{
+	OutTools.Empty();
+	for (const auto& MergeActorTool : MergeActorsTools)
+	{
+		check(MergeActorTool.Get() != nullptr);
+		OutTools.Add(MergeActorTool.Get());
+	}
+}
 
 #undef LOCTEXT_NAMESPACE
diff --git a/Engine/Source/Editor/MergeActors/Private/MergeActorsTool.cpp b/Engine/Source/Editor/MergeActors/Private/MergeActorsTool.cpp
new file mode 100644
--- /dev/null
+++ b/Engine/Source/Editor/MergeActors/Private/MergeActorsTool.cpp
@@ -0,0 +1,76 @@
+// Copyright Epic Games, Inc. All Rights Reserved.
+
+#include "MergeActorsTool.h"
+
+bool FMergeActorsTool::GetReplaceSourceActors() const
+{
+	return bReplaceSourceActors;
+}
+
+void FMergeActorsTool::SetReplaceSourceActors(bool bInReplaceSourceActors)
+{
+	bReplaceSourceActors = bInReplaceSourceActors;
+}
+
+bool FMergeActorsTool::RunMergeFromSelection()
+{
+	TArray<TSharedPtr<FMergeComponentData>> SelectionData;
+	BuildMergeComponentDataFromSelection(SelectionData, bAllowShapeComponents);
+
+	if (SelectionData.Num() == 0)
+	{
+		return false;
+	}
+
+	FString PackageName;
+	if (GetPackageNameForMergeAction(GetDefaultPackageName(), PackageName))
+	{
+		return RunMerge(PackageName, SelectionData);
+	}
+	else
+	{
+		return false;
+	}
+}
+
+bool FMergeActorsTool::RunMergeFromWidget()
+{
+	FString PackageName;
+	if (GetPackageNameForMergeAction(GetDefaultPackageName(), PackageName))
+	{
+		return RunMerge(PackageName, GetSelectedComponentsInWidget());
+	}
+	else
+	{
+		return false;
+	}
+}
+
+bool HasAtLeastOneStaticMesh(const TArray<TSharedPtr<FMergeComponentData>>& ComponentsData)
+{
+	for (const TSharedPtr<FMergeComponentData>& ComponentData : ComponentsData)
+	{
+		if (!ComponentData->bShouldIncorporate)
+			continue;
+
+		const bool bIsMesh = (Cast<UStaticMeshComponent>(ComponentData->PrimComponent.Get()) != nullptr);
+
+		if (bIsMesh)
+			return true;
+	}
+
+	return false;
+}
+
+bool FMergeActorsTool::CanMergeFromSelection() const
+{
+	TArray<TSharedPtr<FMergeComponentData>> SelectedComponents;
+	BuildMergeComponentDataFromSelection(SelectedComponents, bAllowShapeComponents);
+	return HasAtLeastOneStaticMesh(SelectedComponents);
+}
+
+bool FMergeActorsTool::CanMergeFromWidget() const
+{
+	const TArray<TSharedPtr<FMergeComponentData>>& SelectedComponents = GetSelectedComponentsInWidget();
+	return HasAtLeastOneStaticMesh(SelectedComponents);
+}
\ No newline at end of file
diff --git a/Engine/Source/Editor/MergeActors/Private/MergeActorsTool.h b/Engine/Source/Editor/MergeActors/Private/MergeActorsTool.h
new file mode 100644
--- /dev/null
+++ b/Engine/Source/Editor/MergeActors/Private/MergeActorsTool.h
@@ -0,0 +1,40 @@
+// Copyright Epic Games, Inc. All Rights Reserved.
+#pragma once
+
+#include "IMergeActorsTool.h"
+#include "MergeProxyUtils/Utils.h"
+
+/**
+ * Merge Actors tool base class
+ */
+class FMergeActorsTool : public IMergeActorsTool
+{
+public:
+	/** IMergeActorsTool partial implementation */
+	bool GetReplaceSourceActors() const override;
+
+	void SetReplaceSourceActors(bool bReplaceSourceActors) override;
+
+	bool RunMergeFromSelection() override;
+
+	bool RunMergeFromWidget() override;
+
+	bool CanMergeFromSelection() const override;
+
+	bool CanMergeFromWidget() const override;
+
+protected:
+	/**
+	 * Perform the tool merge operation
+	 */
+	virtual bool RunMerge(const FString& PackageName, const TArray<TSharedPtr<FMergeComponentData>>& SelectedComponents) = 0;
+
+	/**
+	 * Returns the selected components from the tool's widget
+	 */
+	virtual const TArray<TSharedPtr<FMergeComponentData>>& GetSelectedComponentsInWidget() const = 0;
+
+protected:
+	bool bReplaceSourceActors = false;
+	bool bAllowShapeComponents = true;
+};
\ No newline at end of file
diff --git a/Engine/Source/Editor/MergeActors/Private/MergeProxyUtils/SMeshProxyCommonDialog.cpp b/Engine/Source/Editor/MergeActors/Private/MergeProxyUtils/SMeshProxyCommonDialog.cpp
new file mode 100644
--- /dev/null
+++ b/Engine/Source/Editor/MergeActors/Private/MergeProxyUtils/SMeshProxyCommonDialog.cpp
@@ -0,0 +1,243 @@
+// Copyright Epic Games, Inc. All Rights Reserved.
+
+#include "MergeProxyUtils/SMeshProxyCommonDialog.h"
+
+#include "EditorStyleSet.h"
+#include "PropertyEditorModule.h"
+#include "IDetailChildrenBuilder.h"
+#include "Modules/ModuleManager.h"
+#include "SlateOptMacros.h"
+#include "Widgets/Text/STextBlock.h"
+#include "Widgets/Views/SListView.h"
+#include "Widgets/Input/SCheckBox.h"
+#include "Engine/Selection.h"
+#include "Editor.h"
+
+#include "IDetailsView.h"
+
+// TODO: review this, have new namespace? text is same as SMeshInstancingDialog
+#define LOCTEXT_NAMESPACE "SMeshMergingDialog"
+
+//////////////////////////////////////////////////////////////////////////
+// SMeshProxyCommonDialog
+
+SMeshProxyCommonDialog::SMeshProxyCommonDialog()
+{
+	bRefreshListView = false;
+	MergeStaticMeshComponentsLabel = LOCTEXT("MergeStaticMeshComponentsLabel", "Mesh Components to be incorporated in the merge:");
+	SelectedComponentsListBoxToolTip = LOCTEXT("SelectedComponentsListBoxToolTip", "The selected mesh components will be incorporated into the merged mesh");
+	DeleteUndoLabel = LOCTEXT("DeleteUndo", "Insufficient mesh components found for merging.");
+}
+
+SMeshProxyCommonDialog::~SMeshProxyCommonDialog()
+{
+	// Remove all delegates
+	USelection::SelectionChangedEvent.RemoveAll(this);
+	USelection::SelectObjectEvent.RemoveAll(this);
+	FEditorDelegates::MapChange.RemoveAll(this);
+	FEditorDelegates::NewCurrentLevel.RemoveAll(this);
+}
+
+BEGIN_SLATE_FUNCTION_BUILD_OPTIMIZATION
+void SMeshProxyCommonDialog::Construct(const FArguments& InArgs)
+{
+	UpdateSelectedStaticMeshComponents();
+	CreateSettingsView();
+
+	// Create widget layout
+	this->ChildSlot
+	[
+		SNew(SVerticalBox)
+		+ SVerticalBox::Slot()
+		[
+			SNew(SHorizontalBox)
+			// Left-side: the component selection pane
+			+ SHorizontalBox::Slot()
+			.AutoWidth()
+			[
+				SNew(SBorder)
+				// We want to emulate the style of the table we'll place here
+				.BorderImage(&FEditorStyle::GetWidgetStyle<FTableRowStyle>("TableView.Row").EvenRowBackgroundBrush)
+				[
+					SNew(SVerticalBox)
+					// Selected components count
+					+ SVerticalBox::Slot()
+					.AutoHeight()
+					.Padding(FEditorStyle::GetMargin("StandardDialog.ContentPadding"))
+					[
+						SNew(SHorizontalBox)
+						+ SHorizontalBox::Slot()
+						.AutoWidth()
+						.VAlign(VAlign_Center)
+						[
+							SNew(STextBlock)
+							.Text(MergeStaticMeshComponentsLabel)
+						]
+						+ SHorizontalBox::Slot()
+						.AutoWidth()
+						.VAlign(VAlign_Center)
+						.HAlign(HAlign_Left)
+						.Padding(FEditorStyle::GetMargin("StandardDialog.SlotPadding"))
+						[
+							SNew(STextBlock)
+							.Text_Lambda([this]() -> FText
+							{
+								return FText::AsNumber(ComponentSelectionControl.NumSelectedMeshComponents);
+							})
+						]
+					]
+					// List of selected/unselected components
+					+ SVerticalBox::Slot()
+					.AutoHeight()
+					.Padding(FEditorStyle::GetMargin("StandardDialog.ContentPadding") + FEditorStyle::GetMargin("StandardDialog.SlotPadding"))
+					[
+						SAssignNew(ComponentSelectionControl.ComponentsListView, SListView<TSharedPtr<FMergeComponentData>>)
+						.ListItemsSource(&ComponentSelectionControl.SelectedComponents)
+						.OnGenerateRow(this, &SMeshProxyCommonDialog::MakeComponentListItemWidget)
+						.ToolTipText(SelectedComponentsListBoxToolTip)
+					]
+				]
+			]
+			// Right-side: the settings view
+			+ SHorizontalBox::Slot()
+			[
+				SNew(SBorder)
+				.BorderImage(FEditorStyle::GetBrush("ToolPanel.GroupBorder"))
+				[
+					SettingsView->AsShared()
+				]
+			]
+		]
+
+		// Diagnostics bar
+		+ SVerticalBox::Slot()
+		.AutoHeight()
+		.Padding(10)
+		[
+			SNew(SBorder)
+			.BorderBackgroundColor(FLinearColor::Yellow)
+			.BorderImage(FEditorStyle::GetBrush("ToolPanel.GroupBorder"))
+			.Visibility_Lambda([this]()->EVisibility { return this->GetContentEnabledState() ? EVisibility::Collapsed : EVisibility::Visible; })
+			[
+				SNew(STextBlock)
+				.Text(DeleteUndoLabel)
+			]
+		]
+
+		// Predictive results bar
+		+ SVerticalBox::Slot()
+		.AutoHeight()
+		.Padding(10)
+		[
+			SNew(SBorder)
+			.BorderBackgroundColor(FLinearColor::Green)
+			.BorderImage(FEditorStyle::GetBrush("ToolPanel.GroupBorder"))
+			.Visibility_Lambda([this]()->EVisibility { return this->GetPredictedResultsText().IsEmpty() ? EVisibility::Collapsed : EVisibility::Visible; })
+			[
+				SNew(STextBlock)
+				.Text(this, &SMeshProxyCommonDialog::GetPredictedResultsText)
+			]
+		]
+	];
+
+	// Selection change
+	USelection::SelectionChangedEvent.AddRaw(this, &SMeshProxyCommonDialog::OnLevelSelectionChanged);
+	USelection::SelectObjectEvent.AddRaw(this, &SMeshProxyCommonDialog::OnLevelSelectionChanged);
+	FEditorDelegates::MapChange.AddSP(this, &SMeshProxyCommonDialog::OnMapChange);
+	FEditorDelegates::NewCurrentLevel.AddSP(this, &SMeshProxyCommonDialog::OnNewCurrentLevel);
+}
+
+void SMeshProxyCommonDialog::OnMapChange(uint32 MapFlags)
+{
+	Reset();
+}
+
+void SMeshProxyCommonDialog::OnNewCurrentLevel()
+{
+	Reset();
+}
+
+void SMeshProxyCommonDialog::Tick(const FGeometry& AllottedGeometry, const double InCurrentTime, const float InDeltaTime)
+{
+	// Check if we need to update selected components and the listbox
+	if(bRefreshListView == true)
+	{
+		ComponentSelectionControl.UpdateSelectedCompnentsAndListBox();
+		bRefreshListView = false;
+	}
+}
+
+void SMeshProxyCommonDialog::Reset()
+{
+	bRefreshListView = true;
+}
+
+bool SMeshProxyCommonDialog::GetContentEnabledState() const
+{
+	return (GetNumSelectedMeshComponents() >= 1); // Only enabled if a mesh is selected
+}
+
+void SMeshProxyCommonDialog::UpdateSelectedStaticMeshComponents()
+{
+	ComponentSelectionControl.UpdateSelectedStaticMeshComponents();
+}
+
+TSharedRef<ITableRow> SMeshProxyCommonDialog::MakeComponentListItemWidget(TSharedPtr<FMergeComponentData> ComponentData, const TSharedRef<STableViewBase>& OwnerTable)
+{
+	return ComponentSelectionControl.MakeComponentListItemWidget(ComponentData, OwnerTable);
+}
+
+
+void SMeshProxyCommonDialog::CreateSettingsView()
+{
+	// Create a property view
+	FPropertyEditorModule& EditModule = FModuleManager::Get().GetModuleChecked<FPropertyEditorModule>("PropertyEditor");
+
+	FDetailsViewArgs DetailsViewArgs;
+	DetailsViewArgs.bUpdatesFromSelection = true;
+	DetailsViewArgs.bLockable = true;
+	DetailsViewArgs.bCustomNameAreaLocation = false;
+	DetailsViewArgs.NameAreaSettings = FDetailsViewArgs::ENameAreaSettings::HideNameArea;
+	DetailsViewArgs.bCustomFilterAreaLocation = true;
+	DetailsViewArgs.DefaultsOnlyVisibility = EEditDefaultsOnlyNodeVisibility::Hide;
+
+	SettingsView = EditModule.CreateDetailView(DetailsViewArgs);
+
+	// Tiny hack to hide this setting, since we have no way / value to go off to 
+	struct Local
+	{
+		/** Delegate to show all properties */
+		static bool IsPropertyVisible(const FPropertyAndParent& PropertyAndParent, bool bInShouldShowNonEditable)
+		{
+			return (PropertyAndParent.Property.GetFName() != GET_MEMBER_NAME_CHECKED(FMaterialProxySettings, GutterSpace));
+		}
+	};
+	
+	SettingsView->SetIsPropertyVisibleDelegate(FIsPropertyVisible::CreateStatic(&Local::IsPropertyVisible, true));
+	
+	SettingsView->OnFinishedChangingProperties().AddSP(this, &SMeshProxyCommonDialog::OnSettingChanged);
+}
+
+END_SLATE_FUNCTION_BUILD_OPTIMIZATION
+
+void SMeshProxyCommonDialog::OnLevelSelectionChanged(UObject* Obj)
+{
+	Reset();
+}
+
+void SMeshProxyCommonDialog::OnSettingChanged(const FPropertyChangedEvent& InPropertyChangedEvent)
+{
+	Reset();
+}
+
+FText SMeshProxyCommonDialog::GetPredictedResultsTextInternal() const
+{
+	return FText();
+}
+
+FText SMeshProxyCommonDialog::GetPredictedResultsText() const
+{
+	return GetPredictedResultsTextInternal();
+}
+
+#undef LOCTEXT_NAMESPACE
\ No newline at end of file
diff --git a/Engine/Source/Editor/MergeActors/Private/MergeProxyUtils/SMeshProxyCommonDialog.h b/Engine/Source/Editor/MergeActors/Private/MergeProxyUtils/SMeshProxyCommonDialog.h
new file mode 100644
--- /dev/null
+++ b/Engine/Source/Editor/MergeActors/Private/MergeProxyUtils/SMeshProxyCommonDialog.h
@@ -0,0 +1,85 @@
+// Copyright Epic Games, Inc. All Rights Reserved.
+#pragma once
+
+#include "CoreMinimal.h"
+#include "Widgets/DeclarativeSyntaxSupport.h"
+#include "Widgets/SCompoundWidget.h"
+#include "Widgets/Views/STableViewBase.h"
+#include "Widgets/Views/STableRow.h"
+
+#include "MergeProxyUtils/Utils.h"
+
+class IDetailsView;
+
+/*-----------------------------------------------------------------------------
+   SMeshProxyCommonDialog
+-----------------------------------------------------------------------------*/
+class SMeshProxyCommonDialog : public SCompoundWidget
+{
+public:
+	SLATE_BEGIN_ARGS(SMeshProxyCommonDialog)
+	{
+	}
+
+	SLATE_END_ARGS()
+
+public:
+	/** **/
+	SMeshProxyCommonDialog();
+	~SMeshProxyCommonDialog();
+
+	/** SWidget functions */
+	void Construct(const FArguments& InArgs);
+
+	/** Getter functionality */
+	const TArray<TSharedPtr<FMergeComponentData>>& GetSelectedComponents() const { return ComponentSelectionControl.SelectedComponents; }
+	/** Get number of selected meshes */
+	const int32 GetNumSelectedMeshComponents() const { return ComponentSelectionControl.NumSelectedMeshComponents; }
+
+	/** Resets the state of the UI and flags it for refreshing */
+	void Reset();
+
+protected:
+	/** Predicted results of the merge given the current settings */
+	virtual FText GetPredictedResultsTextInternal() const;
+
+private:
+	/** Begin override SCompoundWidget */
+	virtual void Tick(const FGeometry& AllottedGeometry, const double InCurrentTime, const float InDeltaTime) override;
+	/** End override SCompoundWidget */
+
+	/** Creates and sets up the settings view element*/
+	void CreateSettingsView();
+
+	/** Delegate for the creation of the list view item's widget */
+	TSharedRef<ITableRow> MakeComponentListItemWidget(TSharedPtr<FMergeComponentData> ComponentData, const TSharedRef<STableViewBase>& OwnerTable);
+
+	/** Delegate to determine whether or not the UI elements should be enabled (determined by number of selected actors / mesh components) */
+	bool GetContentEnabledState() const;
+
+	/** Editor delegates for map and selection changes */
+	void OnLevelSelectionChanged(UObject* Obj);
+	void OnMapChange(uint32 MapFlags);
+	void OnNewCurrentLevel();
+	void OnSettingChanged(const FPropertyChangedEvent& InPropertyChangedEvent);
+
+	/** Delegates for predicted results display */
+	FText GetPredictedResultsText() const;
+
+	/** Updates the selection control */
+	void UpdateSelectedStaticMeshComponents();
+
+protected:
+	FComponentSelectionControl ComponentSelectionControl;
+
+	/** Settings view ui element ptr */
+	TSharedPtr<IDetailsView> SettingsView;
+
+	/** List view state tracking data */
+	bool bRefreshListView;
+
+	/** Labels for various parts of the dialog */
+	FText MergeStaticMeshComponentsLabel;
+	FText SelectedComponentsListBoxToolTip;
+	FText DeleteUndoLabel;
+};
\ No newline at end of file
diff --git a/Engine/Source/Editor/MergeActors/Private/MergeProxyUtils/Utils.cpp b/Engine/Source/Editor/MergeActors/Private/MergeProxyUtils/Utils.cpp
--- a/Engine/Source/Editor/MergeActors/Private/MergeProxyUtils/Utils.cpp
+++ b/Engine/Source/Editor/MergeActors/Private/MergeProxyUtils/Utils.cpp
@@ -2,6 +2,10 @@
 
 #include "MergeProxyUtils/Utils.h"
 
+#include "Misc/PackageName.h"
+#include "IContentBrowserSingleton.h"
+#include "ContentBrowserModule.h"
+
 #include "Widgets/DeclarativeSyntaxSupport.h"
 #include "Widgets/SCompoundWidget.h"
 #include "Components/StaticMeshComponent.h"
@@ -15,30 +19,9 @@
 
 #define LOCTEXT_NAMESPACE "MergeProxyDialog"
 
-void FComponentSelectionControl::UpdateSelectedCompnentsAndListBox()
+void BuildMergeComponentDataFromSelection(TArray<TSharedPtr<FMergeComponentData>>& OutComponentsData, bool bAllowShapeComponents)
 {
-	StoreCheckBoxState();
-	UpdateSelectedStaticMeshComponents();
-	ComponentsListView->ClearSelection();
-	ComponentsListView->RequestListRefresh();
-}
-
-void FComponentSelectionControl::StoreCheckBoxState() 
-{
-	StoredCheckBoxStates.Empty();
-
-	// Loop over selected mesh component and store its checkbox state
-	for (TSharedPtr<FMergeComponentData> SelectedComponent : SelectedComponents)
-	{
-		UPrimitiveComponent* PrimComponent = SelectedComponent->PrimComponent.Get();
-		const ECheckBoxState State = SelectedComponent->bShouldIncorporate ? ECheckBoxState::Checked : ECheckBoxState::Unchecked;
-		StoredCheckBoxStates.Add(PrimComponent, State);
-	}
-}
-
-void FComponentSelectionControl::UpdateSelectedStaticMeshComponents()
-{
-	NumSelectedMeshComponents = 0;
+	OutComponentsData.Empty();
 
 	// Retrieve selected actors
 	USelection* SelectedActors = GEditor->GetSelectedActors();
@@ -57,8 +40,6 @@ void FComponentSelectionControl::UpdateSelectedStaticMeshComponents()
 		}
 	}
 
-	// Retrieve static mesh components from selected actors
-	SelectedComponents.Empty();
 	for (AActor* Actor : Actors)
 	{
 		check(Actor != nullptr);
@@ -78,35 +59,115 @@ void FComponentSelectionControl::UpdateSelectedStaticMeshComponents()
 			}
 			else if (UShapeComponent* ShapeComponent = Cast<UShapeComponent>(PrimComponent))
 			{
-				bShouldIncorporate = true;
-				bInclude = true;
+				if (bAllowShapeComponents)
+				{
+					bShouldIncorporate = true;
+					bInclude = true;
+				}
 			}
 
 			if (bInclude)
 			{
-				SelectedComponents.Add(TSharedPtr<FMergeComponentData>(new FMergeComponentData(PrimComponent)));
-				TSharedPtr<FMergeComponentData>& ComponentData = SelectedComponents.Last();
-
+				OutComponentsData.Add(TSharedPtr<FMergeComponentData>(new FMergeComponentData(PrimComponent)));
+				TSharedPtr<FMergeComponentData>& ComponentData = OutComponentsData.Last();
 				ComponentData->bShouldIncorporate = bShouldIncorporate;
-
-				// See if we stored a checkbox state for this mesh component, and set accordingly
-				ECheckBoxState* StoredState = StoredCheckBoxStates.Find(PrimComponent);
-				if (StoredState != nullptr)
-				{
-					ComponentData->bShouldIncorporate = (*StoredState == ECheckBoxState::Checked);
-				}
-
-				// Keep count of selected meshes
-				if (ComponentData->bShouldIncorporate && bIsMesh)
-				{
-					NumSelectedMeshComponents++;
-				}
 			}
+		}
+	}
+}
 
+void BuildActorsListFromMergeComponentsData(const TArray<TSharedPtr<FMergeComponentData>>& InComponentsData, TArray<AActor*>& OutActors, TArray<ULevel*>* OutLevels /* = nullptr */)
+{
+	for (const TSharedPtr<FMergeComponentData>& SelectedComponent : InComponentsData)
+	{
+		if (SelectedComponent->PrimComponent.IsValid())
+		{
+			AActor* Actor = SelectedComponent->PrimComponent.Get()->GetOwner();
+			OutActors.AddUnique(Actor);
+			if (OutLevels)
+				OutLevels->AddUnique(Actor->GetLevel());
 		}
 	}
 }
 
+bool GetPackageNameForMergeAction(const FString& DefaultPackageName, FString& OutPackageName)
+{
+	if (DefaultPackageName.Len() > 0)
+	{
+		const FString DefaultPath = FPackageName::GetLongPackagePath(DefaultPackageName);
+		const FString DefaultName = FPackageName::GetShortName(DefaultPackageName);
+
+		// Initialize SaveAssetDialog config
+		FSaveAssetDialogConfig SaveAssetDialogConfig;
+		SaveAssetDialogConfig.DialogTitleOverride = LOCTEXT("CreateMergedActorTitle", "Create Merged Actor");
+		SaveAssetDialogConfig.DefaultPath = DefaultPath;
+		SaveAssetDialogConfig.DefaultAssetName = DefaultName;
+		SaveAssetDialogConfig.ExistingAssetPolicy = ESaveAssetDialogExistingAssetPolicy::AllowButWarn;
+		SaveAssetDialogConfig.AssetClassNames = { UStaticMesh::StaticClass()->GetFName() };
+
+		FContentBrowserModule& ContentBrowserModule = FModuleManager::LoadModuleChecked<FContentBrowserModule>("ContentBrowser");
+		FString SaveObjectPath = ContentBrowserModule.Get().CreateModalSaveAssetDialog(SaveAssetDialogConfig);
+		if (!SaveObjectPath.IsEmpty())
+		{
+			OutPackageName = FPackageName::ObjectPathToPackageName(SaveObjectPath);
+			return true;
+		}
+		else
+		{
+			return false;
+		}
+	}
+	else
+	{
+		OutPackageName = DefaultPackageName;
+		return true;
+	}
+}
+
+void FComponentSelectionControl::UpdateSelectedCompnentsAndListBox()
+{
+	StoreCheckBoxState();
+	UpdateSelectedStaticMeshComponents();
+	ComponentsListView->ClearSelection();
+	ComponentsListView->RequestListRefresh();
+}
+
+void FComponentSelectionControl::StoreCheckBoxState() 
+{
+	StoredCheckBoxStates.Empty();
+
+	// Loop over selected mesh component and store its checkbox state
+	for (TSharedPtr<FMergeComponentData> SelectedComponent : SelectedComponents)
+	{
+		UPrimitiveComponent* PrimComponent = SelectedComponent->PrimComponent.Get();
+		const ECheckBoxState State = SelectedComponent->bShouldIncorporate ? ECheckBoxState::Checked : ECheckBoxState::Unchecked;
+		StoredCheckBoxStates.Add(PrimComponent, State);
+	}
+}
+
+void FComponentSelectionControl::UpdateSelectedStaticMeshComponents()
+{
+	BuildMergeComponentDataFromSelection(SelectedComponents, bAllowShapeComponents);
+
+	// Count number of selected objects & Update check state based on previous data
+	NumSelectedMeshComponents = 0;
+
+	for (TSharedPtr<FMergeComponentData>& ComponentData : SelectedComponents)
+	{
+		ECheckBoxState* StoredState = StoredCheckBoxStates.Find(ComponentData->PrimComponent.Get());
+		if (StoredState != nullptr)
+		{
+			ComponentData->bShouldIncorporate = (*StoredState == ECheckBoxState::Checked);
+		}
+
+		// Keep count of selected meshes
+		const bool bIsMesh = (Cast<UStaticMeshComponent>(ComponentData->PrimComponent.Get()) != nullptr);
+		if (ComponentData->bShouldIncorporate && bIsMesh)
+		{
+			NumSelectedMeshComponents++;
+		}
+	}
+}
 
 TSharedRef<ITableRow> FComponentSelectionControl::MakeComponentListItemWidget(TSharedPtr<FMergeComponentData> ComponentData, const TSharedRef<STableViewBase>& OwnerTable)
 {
diff --git a/Engine/Source/Editor/MergeActors/Private/MergeProxyUtils/Utils.h b/Engine/Source/Editor/MergeActors/Private/MergeProxyUtils/Utils.h
--- a/Engine/Source/Editor/MergeActors/Private/MergeProxyUtils/Utils.h
+++ b/Engine/Source/Editor/MergeActors/Private/MergeProxyUtils/Utils.h
@@ -22,6 +22,10 @@ struct FMergeComponentData
 	bool bShouldIncorporate;
 };
 
+void BuildMergeComponentDataFromSelection(TArray<TSharedPtr<FMergeComponentData>>& OutComponentsData, bool bAllowShapeComponents);
+void BuildActorsListFromMergeComponentsData(const TArray<TSharedPtr<FMergeComponentData>>& InComponentsData, TArray<AActor*>& OutActors, TArray<ULevel*>* OutLevels = nullptr);
+bool GetPackageNameForMergeAction(const FString& DefaultPackageName, FString& OutPackageName);
+
 struct FComponentSelectionControl
 {
 
@@ -46,4 +50,7 @@ struct FComponentSelectionControl
 	
 	/** Number of selected static mesh components */
 	int32 NumSelectedMeshComponents = 0;
+
+	/** Allow selection of ShapeComponents **/
+	bool bAllowShapeComponents = true;
 };
diff --git a/Engine/Source/Editor/MergeActors/Private/MeshInstancingTool/MeshInstancingTool.cpp b/Engine/Source/Editor/MergeActors/Private/MeshInstancingTool/MeshInstancingTool.cpp
--- a/Engine/Source/Editor/MergeActors/Private/MeshInstancingTool/MeshInstancingTool.cpp
+++ b/Engine/Source/Editor/MergeActors/Private/MeshInstancingTool/MeshInstancingTool.cpp
@@ -28,6 +28,7 @@ UMeshInstancingSettingsObject* UMeshInstancingSettingsObject::DefaultSettings =
 
 FMeshInstancingTool::FMeshInstancingTool()
 {
+	bAllowShapeComponents = false;
 	SettingsObject = UMeshInstancingSettingsObject::Get();
 }
 
@@ -45,7 +46,7 @@ TSharedRef<SWidget> FMeshInstancingTool::GetWidget()
 
 FText FMeshInstancingTool::GetTooltipText() const
 {
-	return LOCTEXT("MeshInstancingToolTooltip", "Harvest geometry from selected actors and merge them into an actor with multiple instanced static mesh components.");
+	return LOCTEXT("MeshInstancingToolTooltip", "Merge to an actor with instanced static mesh components");
 }
 
 FString FMeshInstancingTool::GetDefaultPackageName() const
@@ -53,21 +54,18 @@ FString FMeshInstancingTool::GetDefaultPackageName() const
 	return FString();
 }
 
-bool FMeshInstancingTool::RunMerge(const FString& PackageName)
+const TArray<TSharedPtr<FMergeComponentData>>& FMeshInstancingTool::GetSelectedComponentsInWidget() const
+{
+	return InstancingDialog->GetSelectedComponents();
+}
+
+bool FMeshInstancingTool::RunMerge(const FString& PackageName, const TArray<TSharedPtr<FMergeComponentData>>& SelectedComponents)
 {
 	const IMeshMergeUtilities& MeshUtilities = FModuleManager::Get().LoadModuleChecked<IMeshMergeModule>("MeshMergeUtilities").GetUtilities();
-	USelection* SelectedActors = GEditor->GetSelectedActors();
 	TArray<AActor*> Actors;
 	TArray<ULevel*> UniqueLevels;
-	for (FSelectionIterator Iter(*SelectedActors); Iter; ++Iter)
-	{
-		AActor* Actor = Cast<AActor>(*Iter);
-		if (Actor)
-		{
-			Actors.Add(Actor);
-			UniqueLevels.AddUnique(Actor->GetLevel());
-		}
-	}
+
+	BuildActorsListFromMergeComponentsData(SelectedComponents, Actors, &UniqueLevels);
 
 	// This restriction is only for replacement of selected actors with merged mesh actor
 	if (UniqueLevels.Num() > 1)
@@ -84,10 +82,9 @@ bool FMeshInstancingTool::RunMerge(const FString& PackageName)
 		SlowTask.MakeDialog();
 
 		// Extracting static mesh components from the selected mesh components in the dialog
-		const TArray<TSharedPtr<FInstanceComponentData>>& SelectedComponents = InstancingDialog->GetSelectedComponents();
 		TArray<UPrimitiveComponent*> ComponentsToMerge;
 
-		for ( const TSharedPtr<FInstanceComponentData>& SelectedComponent : SelectedComponents)
+		for ( const TSharedPtr<FMergeComponentData>& SelectedComponent : SelectedComponents)
 		{
 			// Determine whether or not this component should be incorporated according the user settings
 			if (SelectedComponent->bShouldIncorporate && SelectedComponent->PrimComponent.IsValid())
@@ -101,11 +98,15 @@ bool FMeshInstancingTool::RunMerge(const FString& PackageName)
 			// spawn the actor that will contain out instances
 			UWorld* World = ComponentsToMerge[0]->GetWorld();
 			checkf(World != nullptr, TEXT("Invalid World retrieved from Mesh components"));
-			MeshUtilities.MergeComponentsToInstances(ComponentsToMerge, World, UniqueLevels[0], SettingsObject->Settings);
+			const bool bActuallyMerge = true;
+			MeshUtilities.MergeComponentsToInstances(ComponentsToMerge, World, UniqueLevels[0], SettingsObject->Settings, bActuallyMerge, bReplaceSourceActors, nullptr);
 		}
 	}
 
-	InstancingDialog->Reset();
+	if (InstancingDialog)
+	{
+		InstancingDialog->Reset();
+	}
 
 	return true;
 }
@@ -113,18 +114,11 @@ bool FMeshInstancingTool::RunMerge(const FString& PackageName)
 FText FMeshInstancingTool::GetPredictedResultsText()
 {
 	const IMeshMergeUtilities& MeshUtilities = FModuleManager::Get().LoadModuleChecked<IMeshMergeModule>("MeshMergeUtilities").GetUtilities();
-	USelection* SelectedActors = GEditor->GetSelectedActors();
+	const TArray<TSharedPtr<FMergeComponentData>>& SelectedComponents = InstancingDialog->GetSelectedComponents();
 	TArray<AActor*> Actors;
 	TArray<ULevel*> UniqueLevels;
-	for (FSelectionIterator Iter(*SelectedActors); Iter; ++Iter)
-	{
-		AActor* Actor = Cast<AActor>(*Iter);
-		if (Actor)
-		{
-			Actors.Add(Actor);
-			UniqueLevels.AddUnique(Actor->GetLevel());
-		}
-	}
+
+	BuildActorsListFromMergeComponentsData(SelectedComponents, Actors, &UniqueLevels);
 
 	// This restriction is only for replacement of selected actors with merged mesh actor
 	if (UniqueLevels.Num() > 1)
@@ -133,10 +127,9 @@ FText FMeshInstancingTool::GetPredictedResultsText()
 	}
 
 	// Extracting static mesh components from the selected mesh components in the dialog
-	const TArray<TSharedPtr<FInstanceComponentData>>& SelectedComponents = InstancingDialog->GetSelectedComponents();
 	TArray<UPrimitiveComponent*> ComponentsToMerge;
 
-	for ( const TSharedPtr<FInstanceComponentData>& SelectedComponent : SelectedComponents)
+	for ( const TSharedPtr<FMergeComponentData>& SelectedComponent : SelectedComponents)
 	{
 		// Determine whether or not this component should be incorporated according the user settings
 		if (SelectedComponent->bShouldIncorporate)
@@ -146,11 +139,12 @@ FText FMeshInstancingTool::GetPredictedResultsText()
 	}
 		
 	FText OutResultsText;
-	if(ComponentsToMerge.Num() > 0)
+	if(ComponentsToMerge.Num() > 0 && ComponentsToMerge[0])
 	{
 		UWorld* World = ComponentsToMerge[0]->GetWorld();
 		checkf(World != nullptr, TEXT("Invalid World retrieved from Mesh components"));
-		MeshUtilities.MergeComponentsToInstances(ComponentsToMerge, World, UniqueLevels[0], SettingsObject->Settings, false, &OutResultsText);
+		const bool bActuallyMerge = false;
+		MeshUtilities.MergeComponentsToInstances(ComponentsToMerge, World, UniqueLevels.Num() > 0 ? UniqueLevels[0] : nullptr, SettingsObject->Settings, bActuallyMerge, bReplaceSourceActors, &OutResultsText);
 	}
 	else
 	{
@@ -159,9 +153,4 @@ FText FMeshInstancingTool::GetPredictedResultsText()
 
 	return OutResultsText;
 }
-
-bool FMeshInstancingTool::CanMerge() const
-{	
-	return InstancingDialog->GetNumSelectedMeshComponents() >= 1;
-}
 #undef LOCTEXT_NAMESPACE
\ No newline at end of file
diff --git a/Engine/Source/Editor/MergeActors/Private/MeshInstancingTool/MeshInstancingTool.h b/Engine/Source/Editor/MergeActors/Private/MeshInstancingTool/MeshInstancingTool.h
--- a/Engine/Source/Editor/MergeActors/Private/MeshInstancingTool/MeshInstancingTool.h
+++ b/Engine/Source/Editor/MergeActors/Private/MeshInstancingTool/MeshInstancingTool.h
@@ -6,7 +6,7 @@
 #include "UObject/Object.h"
 #include "Widgets/SWidget.h"
 #include "Engine/MeshMerging.h"
-#include "IMergeActorsTool.h"
+#include "MergeActorsTool.h"
 
 #include "MeshInstancingTool.generated.h"
 
@@ -62,7 +62,7 @@ public:
 /**
  * Mesh Instancing Tool
  */
-class FMeshInstancingTool : public IMergeActorsTool
+class FMeshInstancingTool : public FMergeActorsTool
 {
 	friend class SMeshInstancingDialog;
 
@@ -76,11 +76,14 @@ public:
 	virtual FName GetIconName() const override { return "MergeActors.MeshInstancingTool"; }
 	virtual FText GetTooltipText() const override;
 	virtual FString GetDefaultPackageName() const override;
-	virtual bool CanMerge() const override;
-	virtual bool RunMerge(const FString& PackageName) override;
 
 	/** Runs the merging logic to determine predicted results */
 	FText GetPredictedResultsText();
+
+protected:
+	virtual bool RunMerge(const FString& PackageName, const TArray<TSharedPtr<FMergeComponentData>>& SelectedComponents) override;
+	virtual const TArray<TSharedPtr<FMergeComponentData>>& GetSelectedComponentsInWidget() const override;
+
 private:
 	/** Pointer to the mesh instancing dialog containing settings for the merge */
 	TSharedPtr<SMeshInstancingDialog> InstancingDialog;
diff --git a/Engine/Source/Editor/MergeActors/Private/MeshInstancingTool/SMeshInstancingDialog.cpp b/Engine/Source/Editor/MergeActors/Private/MeshInstancingTool/SMeshInstancingDialog.cpp
--- a/Engine/Source/Editor/MergeActors/Private/MeshInstancingTool/SMeshInstancingDialog.cpp
+++ b/Engine/Source/Editor/MergeActors/Private/MeshInstancingTool/SMeshInstancingDialog.cpp
@@ -2,39 +2,20 @@
 
 #include "MeshInstancingTool/SMeshInstancingDialog.h"
 #include "MeshInstancingTool/MeshInstancingTool.h"
-#include "EditorStyleSet.h"
-#include "PropertyEditorModule.h"
-#include "IDetailChildrenBuilder.h"
-#include "Modules/ModuleManager.h"
 #include "SlateOptMacros.h"
-#include "Widgets/Text/STextBlock.h"
-#include "Widgets/Views/SListView.h"
-#include "Widgets/Input/SCheckBox.h"
-#include "Engine/Selection.h"
-#include "Editor.h"
-#include "Components/ChildActorComponent.h"
-#include "Components/ShapeComponent.h"
 
 #include "IDetailsView.h"
 
-#define LOCTEXT_NAMESPACE "SMeshInstancingDialog"
-
 //////////////////////////////////////////////////////////////////////////
 // SMeshInstancingDialog
-
 SMeshInstancingDialog::SMeshInstancingDialog()
 {
-	bRefreshListView = false;
-	NumSelectedMeshComponents = 0;
+	ComponentSelectionControl.bAllowShapeComponents = false;
 }
 
 SMeshInstancingDialog::~SMeshInstancingDialog()
 {
-	// Remove all delegates
-	USelection::SelectionChangedEvent.RemoveAll(this);
-	USelection::SelectObjectEvent.RemoveAll(this);	
-	FEditorDelegates::MapChange.RemoveAll(this);
-	FEditorDelegates::NewCurrentLevel.RemoveAll(this);
+
 }
 
 BEGIN_SLATE_FUNCTION_BUILD_OPTIMIZATION
@@ -43,103 +24,7 @@ void SMeshInstancingDialog::Construct(const FArguments& InArgs, FMeshInstancingT
 	checkf(InTool != nullptr, TEXT("Invalid owner tool supplied"));
 	Tool = InTool;
 
-	UpdateSelectedStaticMeshComponents();
-	CreateSettingsView();
-	
-	// Create widget layout
-	this->ChildSlot
-	[
-		SNew(SVerticalBox)
-		+ SVerticalBox::Slot()
-		.AutoHeight()
-		.Padding(0, 10, 0, 0)
-		[
-			SNew(SBorder)
-			.BorderImage(FEditorStyle::GetBrush("ToolPanel.GroupBorder"))
-			[
-				SNew(SVerticalBox)
-				// Static mesh component selection
-				+SVerticalBox::Slot()
-				.AutoHeight()
-				.Padding(FEditorStyle::GetMargin("StandardDialog.ContentPadding"))
-				[
-					SNew(SHorizontalBox)
-					+ SHorizontalBox::Slot()
-					.AutoWidth()
-					.VAlign(VAlign_Center)
-					[
-						SNew(STextBlock)
-						.Text(LOCTEXT("MergeStaticMeshComponentsLabel", "Mesh Components to be replaced by instances:"))
-					]
-				]
-				+ SVerticalBox::Slot()
-					.AutoHeight()
-					.Padding(FEditorStyle::GetMargin("StandardDialog.ContentPadding"))
-					[
-						SAssignNew(ComponentsListView, SListView<TSharedPtr<FInstanceComponentData>>)
-						.ListItemsSource(&SelectedComponents)
-						.OnGenerateRow(this, &SMeshInstancingDialog::MakeComponentListItemWidget)
-						.ToolTipText(LOCTEXT("SelectedComponentsListBoxToolTip", "The selected mesh components will be incorporated replaced by instances"))
-					]
-			]
-		]
-
-		+ SVerticalBox::Slot()
-		.Padding(0, 10, 0, 0)
-		[
-			SNew(SBorder)
-			.BorderImage(FEditorStyle::GetBrush("ToolPanel.GroupBorder"))
-			[
-				SNew(SVerticalBox)
-				// Static mesh component selection
-				+ SVerticalBox::Slot()
-				.Padding(FEditorStyle::GetMargin("StandardDialog.ContentPadding"))
-				[
-					SNew(SHorizontalBox)
-					+ SHorizontalBox::Slot()
-					.VAlign(VAlign_Center)
-					[
-						SettingsView->AsShared()
-					]
-				]
-			]
-		]
-
-		+ SVerticalBox::Slot()
-		.AutoHeight()
-		.Padding(10)
-		[
-			SNew(SBorder)
-			.BorderBackgroundColor(FLinearColor::Yellow)
-			.BorderImage(FEditorStyle::GetBrush("ToolPanel.GroupBorder"))
-			.Visibility_Lambda([this]()->EVisibility { return this->GetContentEnabledState() ? EVisibility::Collapsed : EVisibility::Visible; })
-			[
-				SNew(STextBlock)
-				.Text(LOCTEXT("DeleteUndo", "Insufficient mesh components found for instance replacement"))
-			]
-		]
-
-		+ SVerticalBox::Slot()
-		.AutoHeight()
-		.Padding(10)
-		[
-			SNew(SBorder)
-			.BorderBackgroundColor(FLinearColor::Green)
-			.BorderImage(FEditorStyle::GetBrush("ToolPanel.GroupBorder"))
-			.Visibility_Lambda([this]()->EVisibility { return this->GetPredictedResultsText().IsEmpty() ? EVisibility::Collapsed : EVisibility::Visible; })
-			[
-				SNew(STextBlock)
-				.Text(this, &SMeshInstancingDialog::GetPredictedResultsText)
-			]
-		]
-	];	
-
-
-	// Selection change
-	USelection::SelectionChangedEvent.AddRaw(this, &SMeshInstancingDialog::OnLevelSelectionChanged);
-	USelection::SelectObjectEvent.AddRaw(this, &SMeshInstancingDialog::OnLevelSelectionChanged);
-	FEditorDelegates::MapChange.AddSP(this, &SMeshInstancingDialog::OnMapChange);
-	FEditorDelegates::NewCurrentLevel.AddSP(this, &SMeshInstancingDialog::OnNewCurrentLevel);
+	SMeshProxyCommonDialog::Construct(SMeshProxyCommonDialog::FArguments());
 
 	InstancingSettings = UMeshInstancingSettingsObject::Get();
 	SettingsView->SetObject(InstancingSettings);
@@ -147,236 +32,9 @@ void SMeshInstancingDialog::Construct(const FArguments& InArgs, FMeshInstancingT
 	Reset();
 }
 
-void SMeshInstancingDialog::OnMapChange(uint32 MapFlags)
-{
-	Reset();
-}
-
-void SMeshInstancingDialog::OnNewCurrentLevel()
-{
-	Reset();
-}
-
-
-void SMeshInstancingDialog::OnSettingChanged(const FPropertyChangedEvent& InPropertyChangedEvent)
-{
-	Reset();
-}
-
-void SMeshInstancingDialog::Tick(const FGeometry& AllottedGeometry, const double InCurrentTime, const float InDeltaTime)
-{
-	// Check if we need to update selected components and the listbox
-	if (bRefreshListView == true)
-	{
-		StoreCheckBoxState();
-		UpdateSelectedStaticMeshComponents();
-		RefreshPredictedResultsText();
-		ComponentsListView->ClearSelection();
-		ComponentsListView->RequestListRefresh();
-		bRefreshListView = false;		
-	}
-}
-
-void SMeshInstancingDialog::Reset()
-{	
-	bRefreshListView = true;
-}
-
-bool SMeshInstancingDialog::GetContentEnabledState() const
-{
-	return (GetNumSelectedMeshComponents() >= 1); // Only enabled if a mesh is selected
-}
-
-void SMeshInstancingDialog::UpdateSelectedStaticMeshComponents()
-{	
-	NumSelectedMeshComponents = 0;
-
-	// Retrieve selected actors
-	USelection* SelectedActors = GEditor->GetSelectedActors();
-	TArray<AActor*> Actors;
-	TArray<ULevel*> UniqueLevels;
-	for (FSelectionIterator Iter(*SelectedActors); Iter; ++Iter)
-	{
-		AActor* Actor = Cast<AActor>(*Iter);
-		if (Actor)
-		{
-			Actors.Add(Actor);
-			UniqueLevels.AddUnique(Actor->GetLevel());
-		}
-	}
-
-	// Retrieve static mesh components from selected actors
-	SelectedComponents.Empty();
-	for (int32 ActorIndex = 0; ActorIndex < Actors.Num(); ++ActorIndex )
-	{
-		AActor* Actor = Actors[ActorIndex];
-		check(Actor != nullptr);
-
-		TArray<UChildActorComponent*> ChildActorComponents;
-		Actor->GetComponents<UChildActorComponent>(ChildActorComponents);
-		for (UChildActorComponent* ChildComponent : ChildActorComponents)
-		{
-			// Push actor at the back of array so we will process it
-			AActor* ChildActor = ChildComponent->GetChildActor();
-			if (ChildActor)
-			{
-				Actors.Add(ChildActor);
-			}
-		}
-		
-		TArray<UPrimitiveComponent*> PrimComponents;
-		Actor->GetComponents<UPrimitiveComponent>(PrimComponents);
-		for (UPrimitiveComponent* PrimComponent : PrimComponents)
-		{
-			bool bInclude = false; // Should put into UI list
-			bool bShouldIncorporate = false; // Should default to part of merged mesh
-			bool bIsMesh = false;
-			if (UStaticMeshComponent* StaticMeshComponent = Cast<UStaticMeshComponent>(PrimComponent))
-			{
-				bShouldIncorporate = (StaticMeshComponent->GetStaticMesh() != nullptr);
-				bInclude = true;
-				bIsMesh = true;
-			}
-
-			if (bInclude)
-			{
-				SelectedComponents.Add(TSharedPtr<FInstanceComponentData>(new FInstanceComponentData(PrimComponent)));
-				TSharedPtr<FInstanceComponentData>& ComponentData = SelectedComponents.Last();
-
-				ComponentData->bShouldIncorporate = bShouldIncorporate;
-
-				// See if we stored a checkbox state for this mesh component, and set accordingly
-				ECheckBoxState* StoredState = StoredCheckBoxStates.Find(PrimComponent);
-				if (StoredState != nullptr)
-				{
-					ComponentData->bShouldIncorporate = (*StoredState == ECheckBoxState::Checked);
-				}
-
-				// Keep count of selected meshes
-				if (ComponentData->bShouldIncorporate && bIsMesh)
-				{
-					NumSelectedMeshComponents++;
-				}
-			}
-
-		}		
-	}
-}
-
-TSharedRef<ITableRow> SMeshInstancingDialog::MakeComponentListItemWidget(TSharedPtr<FInstanceComponentData> ComponentData, const TSharedRef<STableViewBase>& OwnerTable)
-{
-	check(ComponentData->PrimComponent != nullptr);
-
-	// Retrieve information about the mesh component
-	const FString OwningActorName = ComponentData->PrimComponent->GetOwner()->GetName();
-
-	// If box should be enabled
-	bool bEnabled = true;
-	bool bIsMesh = false;
-
-	FString ComponentInfo;
-	if (UStaticMeshComponent* StaticMeshComponent = Cast<UStaticMeshComponent>(ComponentData->PrimComponent.Get()))
-	{
-		ComponentInfo = (StaticMeshComponent->GetStaticMesh() != nullptr) ? StaticMeshComponent->GetStaticMesh()->GetName() : TEXT("No Static Mesh Available");
-		bEnabled = (StaticMeshComponent->GetStaticMesh() != nullptr);
-		bIsMesh = true;
-	}
-
-	const FString ComponentName = ComponentData->PrimComponent->GetName();
-
-	// See if we stored a checkbox state for this mesh component, and set accordingly
-	ECheckBoxState State = bEnabled ? ECheckBoxState::Checked : ECheckBoxState::Unchecked;
-	ECheckBoxState* StoredState = StoredCheckBoxStates.Find(ComponentData->PrimComponent.Get());
-	if (StoredState)
-	{
-		State = *StoredState;
-	}
-	
-
-	return SNew(STableRow<TSharedPtr<FInstanceComponentData>>, OwnerTable)
-		[
-			SNew(SBox)
-			[
-				// Disable UI element if this static mesh component has invalid static mesh data
-				SNew(SHorizontalBox)				
-				.IsEnabled(bEnabled)
-				+ SHorizontalBox::Slot()
-				.AutoWidth()
-				[
-					SNew(SCheckBox)
-					.IsChecked(State)
-					.ToolTipText(LOCTEXT("IncorporateCheckBoxToolTip", "When ticked the Component will be incorporated into the merge"))
-					
-					.OnCheckStateChanged_Lambda([=](ECheckBoxState NewState)
-					{
-						ComponentData->bShouldIncorporate = (NewState == ECheckBoxState::Checked);
-
-						if (bIsMesh)
-						{
-							this->NumSelectedMeshComponents += (NewState == ECheckBoxState::Checked) ? 1 : -1;
-						}
-					})
-				]
-
-				+ SHorizontalBox::Slot()
-				.Padding(5.0, 0, 0, 0)
-				.AutoWidth()
-				[
-					SNew(STextBlock)
-					.Text(FText::FromString(OwningActorName + " - " + ComponentInfo + " - " + ComponentName))
-				]
-			]
-		];
-
-}
-
-void SMeshInstancingDialog::CreateSettingsView()
-{
-	// Create a property view
-	FPropertyEditorModule& EditModule = FModuleManager::Get().GetModuleChecked<FPropertyEditorModule>("PropertyEditor");
-
-	FDetailsViewArgs DetailsViewArgs;
-	DetailsViewArgs.bUpdatesFromSelection = true;
-	DetailsViewArgs.bLockable = true;
-	DetailsViewArgs.NameAreaSettings = FDetailsViewArgs::ComponentsAndActorsUseNameArea;
-	DetailsViewArgs.bCustomNameAreaLocation = false;
-	DetailsViewArgs.NameAreaSettings = FDetailsViewArgs::ENameAreaSettings::HideNameArea;
-	DetailsViewArgs.bCustomFilterAreaLocation = true;
-	DetailsViewArgs.DefaultsOnlyVisibility = EEditDefaultsOnlyNodeVisibility::Hide;
-	
-	SettingsView = EditModule.CreateDetailView(DetailsViewArgs);
-
-	SettingsView->OnFinishedChangingProperties().AddSP(this, &SMeshInstancingDialog::OnSettingChanged);
-}
-
 END_SLATE_FUNCTION_BUILD_OPTIMIZATION
 
-void SMeshInstancingDialog::OnLevelSelectionChanged(UObject* Obj)
+FText SMeshInstancingDialog::GetPredictedResultsTextInternal() const
 {
-	Reset();
+	return Tool->GetPredictedResultsText();
 }
-
-void SMeshInstancingDialog::StoreCheckBoxState()
-{
-	StoredCheckBoxStates.Empty();
-
-	// Loop over selected mesh component and store its checkbox state
-	for (TSharedPtr<FInstanceComponentData> SelectedComponent : SelectedComponents )
-	{
-		UPrimitiveComponent* PrimComponent = SelectedComponent->PrimComponent.Get();
-		const ECheckBoxState State = SelectedComponent->bShouldIncorporate ? ECheckBoxState::Checked : ECheckBoxState::Unchecked;
-		StoredCheckBoxStates.Add(PrimComponent, State);
-	}
-}
-
-void SMeshInstancingDialog::RefreshPredictedResultsText()
-{
-	PredictedResultsText = Tool->GetPredictedResultsText();
-}
-
-FText SMeshInstancingDialog::GetPredictedResultsText() const
-{
-	return PredictedResultsText;
-}
-
-#undef LOCTEXT_NAMESPACE
diff --git a/Engine/Source/Editor/MergeActors/Private/MeshInstancingTool/SMeshInstancingDialog.h b/Engine/Source/Editor/MergeActors/Private/MeshInstancingTool/SMeshInstancingDialog.h
--- a/Engine/Source/Editor/MergeActors/Private/MeshInstancingTool/SMeshInstancingDialog.h
+++ b/Engine/Source/Editor/MergeActors/Private/MeshInstancingTool/SMeshInstancingDialog.h
@@ -8,28 +8,15 @@
 #include "Widgets/Views/STableViewBase.h"
 #include "Widgets/Views/STableRow.h"
 
+#include "MergeProxyUtils/SMeshProxyCommonDialog.h"
+
 class FMeshInstancingTool;
-class IDetailsView;
 class UMeshInstancingSettingsObject;
 
-/** Data structure used to keep track of the selected mesh components, and whether or not they should be incorporated in the merge */
-struct FInstanceComponentData
-{
-	FInstanceComponentData(UPrimitiveComponent* InPrimComponent)
-		: PrimComponent(InPrimComponent)
-		, bShouldIncorporate( true )
-	{}
-
-	/** Component extracted from selected actors */
-	TWeakObjectPtr<UPrimitiveComponent> PrimComponent;
-	/** Flag determining whether or not this component should be incorporated into the merge */
-	bool bShouldIncorporate;
-};
-
 /*-----------------------------------------------------------------------------
    SMeshInstancingDialog
 -----------------------------------------------------------------------------*/
-class SMeshInstancingDialog : public SCompoundWidget
+class SMeshInstancingDialog : public SMeshProxyCommonDialog
 {
 public:
 	SLATE_BEGIN_ARGS(SMeshInstancingDialog)
@@ -44,65 +31,16 @@ public:
 	~SMeshInstancingDialog();
 
 	/** SWidget functions */
-	void Construct(const FArguments& InArgs, FMeshInstancingTool* InTool);	
-	
-	/** Getter functionality */
-	const TArray<TSharedPtr<FInstanceComponentData>>& GetSelectedComponents() const { return SelectedComponents; }
-	/** Get number of selected meshes */
-	const int32 GetNumSelectedMeshComponents() const { return NumSelectedMeshComponents; }
-
-	/** Resets the state of the UI and flags it for refreshing */
-	void Reset();
-private:	
-	/** Begin override SCompoundWidget */
-	virtual void Tick(const FGeometry& AllottedGeometry, const double InCurrentTime, const float InDeltaTime) override;
-	/** End override SCompoundWidget */
-
-	/** Creates and sets up the settings view element*/
-	void CreateSettingsView();	
-
-	/** Delegate for the creation of the list view item's widget */
-	TSharedRef<ITableRow> MakeComponentListItemWidget(TSharedPtr<FInstanceComponentData> ComponentData, const TSharedRef<STableViewBase>& OwnerTable);
-
-	/** Delegate to determine whether or not the UI elements should be enabled (determined by number of selected actors / mesh components) */
-	bool GetContentEnabledState() const;
-
-	/** Editor delgates for map and selection changes */
-	void OnLevelSelectionChanged(UObject* Obj);
-	void OnMapChange(uint32 MapFlags);
-	void OnNewCurrentLevel();
-	void OnSettingChanged(const FPropertyChangedEvent& InPropertyChangedEvent);
-
-	/** Updates SelectedMeshComponent array according to retrieved mesh components from editor selection*/
-	void UpdateSelectedStaticMeshComponents();	
-	/** Stores the individual check box states for the currently selected mesh components */
-	void StoreCheckBoxState();
+	void Construct(const FArguments& InArgs, FMeshInstancingTool* InTool);
 
+protected:
 	/** Refresh the predicted results text */;
-	void RefreshPredictedResultsText();
-
-	/** Get the predicted results text */
-	FText GetPredictedResultsText() const;
+	FText GetPredictedResultsTextInternal() const override;
 
 private:
 	/** Owning mesh instancing tool */
 	FMeshInstancingTool* Tool;
-	/** List of mesh components extracted from editor selection */
-	TArray<TSharedPtr<FInstanceComponentData>> SelectedComponents;
-	/** List view ui element */
-	TSharedPtr<SListView<TSharedPtr<FInstanceComponentData>>> ComponentsListView;
-	/** Map of keeping track of checkbox states for each selected component (used to restore state when listview is refreshed) */
-	TMap<UPrimitiveComponent*, ECheckBoxState> StoredCheckBoxStates;
-	/** Settings view ui element ptr */
-	TSharedPtr<IDetailsView> SettingsView;
+
 	/** Cached pointer to mesh instancing setting singleton object */
 	UMeshInstancingSettingsObject* InstancingSettings;
-	/** List view state tracking data */
-	bool bRefreshListView;
-
-	/** Number of selected static mesh components */
-	int32 NumSelectedMeshComponents;
-
-	/** Predicted results of the merge given the current settings */
-	FText PredictedResultsText;
 };
diff --git a/Engine/Source/Editor/MergeActors/Private/MeshMergingTool/MeshMergingTool.cpp b/Engine/Source/Editor/MergeActors/Private/MeshMergingTool/MeshMergingTool.cpp
--- a/Engine/Source/Editor/MergeActors/Private/MeshMergingTool/MeshMergingTool.cpp
+++ b/Engine/Source/Editor/MergeActors/Private/MeshMergingTool/MeshMergingTool.cpp
@@ -28,7 +28,6 @@ UMeshMergingSettingsObject* UMeshMergingSettingsObject::DefaultSettings = nullpt
 bool UMeshMergingSettingsObject::bInitialized = false;
 
 FMeshMergingTool::FMeshMergingTool()
-	: bReplaceSourceActors(false)
 {
 	SettingsObject = UMeshMergingSettingsObject::Get();
 }
@@ -47,7 +46,7 @@ TSharedRef<SWidget> FMeshMergingTool::GetWidget()
 
 FText FMeshMergingTool::GetTooltipText() const
 {
-	return LOCTEXT("MeshMergingToolTooltip", "Harvest geometry from selected actors and merge grouping them by materials.");
+	return LOCTEXT("MeshMergingToolTooltip", "Merge to a single static mesh actor");
 }
 
 FString FMeshMergingTool::GetDefaultPackageName() const
@@ -76,21 +75,18 @@ FString FMeshMergingTool::GetDefaultPackageName() const
 	return PackageName;
 }
 
-bool FMeshMergingTool::RunMerge(const FString& PackageName)
+const TArray<TSharedPtr<FMergeComponentData>>& FMeshMergingTool::GetSelectedComponentsInWidget() const
+{
+	return MergingDialog->GetSelectedComponents();
+}
+
+bool FMeshMergingTool::RunMerge(const FString& PackageName, const TArray<TSharedPtr<FMergeComponentData>>& SelectedComponents)
 {
 	const IMeshMergeUtilities& MeshUtilities = FModuleManager::Get().LoadModuleChecked<IMeshMergeModule>("MeshMergeUtilities").GetUtilities();
-	USelection* SelectedActors = GEditor->GetSelectedActors();
 	TArray<AActor*> Actors;
 	TArray<ULevel*> UniqueLevels;
-	for (FSelectionIterator Iter(*SelectedActors); Iter; ++Iter)
-	{
-		AActor* Actor = Cast<AActor>(*Iter);
-		if (Actor)
-		{
-			Actors.Add(Actor);
-			UniqueLevels.AddUnique(Actor->GetLevel());
-		}
-	}
+
+	BuildActorsListFromMergeComponentsData(SelectedComponents, Actors, bReplaceSourceActors ? &UniqueLevels : nullptr);
 
 	// This restriction is only for replacement of selected actors with merged mesh actor
 	if (UniqueLevels.Num() > 1 && bReplaceSourceActors)
@@ -109,7 +105,6 @@ bool FMeshMergingTool::RunMerge(const FString& PackageName)
 		SlowTask.MakeDialog();
 
 		// Extracting static mesh components from the selected mesh components in the dialog
-		const TArray<TSharedPtr<FMergeComponentData>>& SelectedComponents = MergingDialog->GetSelectedComponents();
 		TArray<UPrimitiveComponent*> ComponentsToMerge;
 
 		for ( const TSharedPtr<FMergeComponentData>& SelectedComponent : SelectedComponents)
@@ -184,16 +179,14 @@ bool FMeshMergingTool::RunMerge(const FString& PackageName)
 		}
 	}
 
-	MergingDialog->Reset();
+	if (MergingDialog)
+	{
+		MergingDialog->Reset();
+	}
 
 	return true;
 }
 
-bool FMeshMergingTool::CanMerge() const
-{	
-	return MergingDialog->GetNumSelectedMeshComponents() >= 1;
-}
-
 #undef LOCTEXT_NAMESPACE
 
 
diff --git a/Engine/Source/Editor/MergeActors/Private/MeshMergingTool/MeshMergingTool.h b/Engine/Source/Editor/MergeActors/Private/MeshMergingTool/MeshMergingTool.h
--- a/Engine/Source/Editor/MergeActors/Private/MeshMergingTool/MeshMergingTool.h
+++ b/Engine/Source/Editor/MergeActors/Private/MeshMergingTool/MeshMergingTool.h
@@ -6,7 +6,7 @@
 #include "UObject/Object.h"
 #include "Widgets/SWidget.h"
 #include "Engine/MeshMerging.h"
-#include "IMergeActorsTool.h"
+#include "MergeActorsTool.h"
 
 #include "MeshMergingTool.generated.h"
 
@@ -67,7 +67,7 @@ public:
 /**
  * Mesh Merging Tool
  */
-class FMeshMergingTool : public IMergeActorsTool
+class FMeshMergingTool : public FMergeActorsTool
 {
 	friend class SMeshMergingDialog;
 
@@ -81,12 +81,12 @@ public:
 	virtual FName GetIconName() const override { return "MergeActors.MeshMergingTool"; }
 	virtual FText GetTooltipText() const override;
 	virtual FString GetDefaultPackageName() const override;
-	virtual bool CanMerge() const override;
-	virtual bool RunMerge(const FString& PackageName) override;	
-private:
-	/** Whether to replace source actors with a merged actor in the world */
-	bool bReplaceSourceActors;
 
+protected:
+	virtual bool RunMerge(const FString& PackageName, const TArray<TSharedPtr<FMergeComponentData>>& SelectedComponents) override;	
+	virtual const TArray<TSharedPtr<FMergeComponentData>>& GetSelectedComponentsInWidget() const override;
+
+private:
 	/** Pointer to the mesh merging dialog containing settings for the merge */
 	TSharedPtr<SMeshMergingDialog> MergingDialog;
 
diff --git a/Engine/Source/Editor/MergeActors/Private/MeshMergingTool/SMeshMergingDialog.cpp b/Engine/Source/Editor/MergeActors/Private/MeshMergingTool/SMeshMergingDialog.cpp
--- a/Engine/Source/Editor/MergeActors/Private/MeshMergingTool/SMeshMergingDialog.cpp
+++ b/Engine/Source/Editor/MergeActors/Private/MeshMergingTool/SMeshMergingDialog.cpp
@@ -2,38 +2,19 @@
 
 #include "MeshMergingTool/SMeshMergingDialog.h"
 #include "MeshMergingTool/MeshMergingTool.h"
-#include "EditorStyleSet.h"
-#include "PropertyEditorModule.h"
-#include "IDetailChildrenBuilder.h"
-#include "Modules/ModuleManager.h"
 #include "SlateOptMacros.h"
-#include "Widgets/Text/STextBlock.h"
-#include "Widgets/Views/SListView.h"
-#include "Widgets/Input/SCheckBox.h"
-#include "Engine/Selection.h"
-#include "Editor.h"
-#include "Components/ChildActorComponent.h"
-#include "Components/ShapeComponent.h"
 
 #include "IDetailsView.h"
 
-#define LOCTEXT_NAMESPACE "SMeshMergingDialog"
-
 //////////////////////////////////////////////////////////////////////////
 // SMeshMergingDialog
 
 SMeshMergingDialog::SMeshMergingDialog()
 {
-	bRefreshListView = false;
 }
 
 SMeshMergingDialog::~SMeshMergingDialog()
 {
-	// Remove all delegates
-	USelection::SelectionChangedEvent.RemoveAll(this);
-	USelection::SelectObjectEvent.RemoveAll(this);	
-	FEditorDelegates::MapChange.RemoveAll(this);
-	FEditorDelegates::NewCurrentLevel.RemoveAll(this);
 }
 
 BEGIN_SLATE_FUNCTION_BUILD_OPTIMIZATION
@@ -42,198 +23,10 @@ void SMeshMergingDialog::Construct(const FArguments& InArgs, FMeshMergingTool* I
 	checkf(InTool != nullptr, TEXT("Invalid owner tool supplied"));
 	Tool = InTool;
 
-	UpdateSelectedStaticMeshComponents();
-	CreateSettingsView();
-	
-	// Create widget layout
-	this->ChildSlot
-	[
-		SNew(SVerticalBox)
-		+ SVerticalBox::Slot()
-		.AutoHeight()
-		.Padding(0, 10, 0, 0)
-		[
-			SNew(SBorder)
-			.BorderImage(FEditorStyle::GetBrush("ToolPanel.GroupBorder"))
-			[
-				SNew(SVerticalBox)
-				// Static mesh component selection
-				+SVerticalBox::Slot()
-				.AutoHeight()
-				.Padding(FEditorStyle::GetMargin("StandardDialog.ContentPadding"))
-				[
-					SNew(SHorizontalBox)
-					+ SHorizontalBox::Slot()
-					.AutoWidth()
-					.VAlign(VAlign_Center)
-					[
-						SNew(STextBlock)
-						.Text(LOCTEXT("MergeStaticMeshComponentsLabel", "Mesh Components to be incorporated in the merge:"))						
-					]
-				]
-				+ SVerticalBox::Slot()
-					.AutoHeight()
-					.Padding(FEditorStyle::GetMargin("StandardDialog.ContentPadding"))
-					[
-						SAssignNew(ComponentSelectionControl.ComponentsListView, SListView<TSharedPtr<FMergeComponentData>>)
-						.ListItemsSource(&ComponentSelectionControl.SelectedComponents)
-					    .OnGenerateRow(this, &SMeshMergingDialog::MakeComponentListItemWidget)
-						.ToolTipText(LOCTEXT("SelectedComponentsListBoxToolTip", "The selected mesh components will be incorporated into the merged mesh"))
-					]
-			]
-		]
-
-		+ SVerticalBox::Slot()
-		.Padding(0, 10, 0, 0)
-		[
-			SNew(SBorder)
-			.BorderImage(FEditorStyle::GetBrush("ToolPanel.GroupBorder"))
-			[
-				SNew(SVerticalBox)
-				// Static mesh component selection
-				+ SVerticalBox::Slot()
-				.Padding(FEditorStyle::GetMargin("StandardDialog.ContentPadding"))
-				[
-					SNew(SHorizontalBox)
-					+ SHorizontalBox::Slot()
-					.VAlign(VAlign_Center)
-					[
-						SettingsView->AsShared()
-					]
-				]
-			]
-		]
-
-		// Replace source actors
-		+ SVerticalBox::Slot()
-		.AutoHeight()
-		.Padding(FEditorStyle::GetMargin("StandardDialog.ContentPadding"))
-		[
-			SNew(SCheckBox)
-			.Type(ESlateCheckBoxType::CheckBox)
-			.IsChecked(this, &SMeshMergingDialog::GetReplaceSourceActors)
-			.OnCheckStateChanged(this, &SMeshMergingDialog::SetReplaceSourceActors)
-			.ToolTipText(LOCTEXT("ReplaceSourceActorsToolTip", "When enabled the Source Actors will be replaced with the newly generated merged mesh"))
-			.Content()
-			[
-				SNew(STextBlock)
-				.Text(LOCTEXT("ReplaceSourceActorsLabel", "Replace Source Actors"))
-				.Font(FEditorStyle::GetFontStyle("StandardDialog.SmallFont"))
-			]
-		]
-
-		+ SVerticalBox::Slot()
-		.AutoHeight()
-		.Padding(10)
-		[
-			SNew(SBorder)
-			.BorderBackgroundColor(FLinearColor::Yellow)
-			.BorderImage(FEditorStyle::GetBrush("ToolPanel.GroupBorder"))
-			.Visibility_Lambda([this]()->EVisibility { return this->GetContentEnabledState() ? EVisibility::Collapsed : EVisibility::Visible; })
-			[
-				SNew(STextBlock)
-				.Text(LOCTEXT("DeleteUndo", "Insufficient mesh components found for merging."))
-			]
-		]
-	];	
-
-
-	// Selection change
-	USelection::SelectionChangedEvent.AddRaw(this, &SMeshMergingDialog::OnLevelSelectionChanged);
-	USelection::SelectObjectEvent.AddRaw(this, &SMeshMergingDialog::OnLevelSelectionChanged);
-	FEditorDelegates::MapChange.AddSP(this, &SMeshMergingDialog::OnMapChange);
-	FEditorDelegates::NewCurrentLevel.AddSP(this, &SMeshMergingDialog::OnNewCurrentLevel);
+	SMeshProxyCommonDialog::Construct(SMeshProxyCommonDialog::FArguments());
 
 	MergeSettings = UMeshMergingSettingsObject::Get();
 	SettingsView->SetObject(MergeSettings);
 }
 
-void SMeshMergingDialog::OnMapChange(uint32 MapFlags)
-{
-	Reset();
-}
-
-void SMeshMergingDialog::OnNewCurrentLevel()
-{
-	Reset();
-}
-
-void SMeshMergingDialog::Tick(const FGeometry& AllottedGeometry, const double InCurrentTime, const float InDeltaTime)
-{
-	// Check if we need to update selected components and the listbox
-	if (bRefreshListView == true)
-	{
-		ComponentSelectionControl.UpdateSelectedCompnentsAndListBox();
-		bRefreshListView = false;		
-	}
-}
-
-void SMeshMergingDialog::Reset()
-{	
-	bRefreshListView = true;
-}
-
-ECheckBoxState SMeshMergingDialog::GetReplaceSourceActors() const
-{
-	return (Tool->bReplaceSourceActors ? ECheckBoxState::Checked : ECheckBoxState::Unchecked);
-}
-
-void SMeshMergingDialog::SetReplaceSourceActors(ECheckBoxState NewValue)
-{
-	Tool->bReplaceSourceActors = (ECheckBoxState::Checked == NewValue);
-}
-
-bool SMeshMergingDialog::GetContentEnabledState() const
-{
-	return (GetNumSelectedMeshComponents() >= 1); // Only enabled if a mesh is selected
-}
-
-void SMeshMergingDialog::UpdateSelectedStaticMeshComponents()
-{	
-	ComponentSelectionControl.UpdateSelectedStaticMeshComponents();
-}
-
-
-TSharedRef<ITableRow> SMeshMergingDialog::MakeComponentListItemWidget(TSharedPtr<FMergeComponentData> ComponentData, const TSharedRef<STableViewBase>& OwnerTable)
-{
-	return ComponentSelectionControl.MakeComponentListItemWidget(ComponentData, OwnerTable);
-}
-
-
-void SMeshMergingDialog::CreateSettingsView()
-{
-	// Create a property view
-	FPropertyEditorModule& EditModule = FModuleManager::Get().GetModuleChecked<FPropertyEditorModule>("PropertyEditor");
-
-	FDetailsViewArgs DetailsViewArgs;
-	DetailsViewArgs.bUpdatesFromSelection = true;
-	DetailsViewArgs.bLockable = true;
-	DetailsViewArgs.NameAreaSettings = FDetailsViewArgs::ComponentsAndActorsUseNameArea;
-	DetailsViewArgs.bCustomNameAreaLocation = false;
-	DetailsViewArgs.NameAreaSettings = FDetailsViewArgs::ENameAreaSettings::HideNameArea;
-	DetailsViewArgs.bCustomFilterAreaLocation = true;
-	DetailsViewArgs.DefaultsOnlyVisibility = EEditDefaultsOnlyNodeVisibility::Hide;
-	
-		
-	// Tiny hack to hide this setting, since we have no way / value to go off to 
-	struct Local
-	{
-		/** Delegate to show all properties */
-		static bool IsPropertyVisible(const FPropertyAndParent& PropertyAndParent, bool bInShouldShowNonEditable)
-		{
-			return (PropertyAndParent.Property.GetFName() != GET_MEMBER_NAME_CHECKED(FMaterialProxySettings, GutterSpace));
-		}
-	};
-
-	SettingsView = EditModule.CreateDetailView(DetailsViewArgs);
-	SettingsView->SetIsPropertyVisibleDelegate(FIsPropertyVisible::CreateStatic(&Local::IsPropertyVisible, true));
-}
-
 END_SLATE_FUNCTION_BUILD_OPTIMIZATION
-
-void SMeshMergingDialog::OnLevelSelectionChanged(UObject* Obj)
-{
-	Reset();
-}
-
-#undef LOCTEXT_NAMESPACE
diff --git a/Engine/Source/Editor/MergeActors/Private/MeshMergingTool/SMeshMergingDialog.h b/Engine/Source/Editor/MergeActors/Private/MeshMergingTool/SMeshMergingDialog.h
--- a/Engine/Source/Editor/MergeActors/Private/MeshMergingTool/SMeshMergingDialog.h
+++ b/Engine/Source/Editor/MergeActors/Private/MeshMergingTool/SMeshMergingDialog.h
@@ -9,15 +9,15 @@
 #include "Widgets/Views/STableRow.h"
 
 #include "MergeProxyUtils/Utils.h"
+#include "MergeProxyUtils/SMeshProxyCommonDialog.h"
 
 class FMeshMergingTool;
-class IDetailsView;
 class UMeshMergingSettingsObject;
 
 /*-----------------------------------------------------------------------------
    SMeshMergingDialog
 -----------------------------------------------------------------------------*/
-class SMeshMergingDialog : public SCompoundWidget
+class SMeshMergingDialog : public SMeshProxyCommonDialog
 {
 public:
 	SLATE_BEGIN_ARGS(SMeshMergingDialog)
@@ -34,53 +34,10 @@ public:
 	/** SWidget functions */
 	void Construct(const FArguments& InArgs, FMeshMergingTool* InTool);	
 
-	/** Getter functionality */
-	const TArray<TSharedPtr<FMergeComponentData>>& GetSelectedComponents() const { return ComponentSelectionControl.SelectedComponents; }
-	/** Get number of selected meshes */
-	const int32 GetNumSelectedMeshComponents() const { return ComponentSelectionControl.NumSelectedMeshComponents; }
-
-	/** Resets the state of the UI and flags it for refreshing */
-	void Reset();
-private:	
-	/** Begin override SCompoundWidget */
-	virtual void Tick(const FGeometry& AllottedGeometry, const double InCurrentTime, const float InDeltaTime) override;
-	/** End override SCompoundWidget */
-
-	/** Creates and sets up the settings view element*/
-	void CreateSettingsView();	
-
-	/** Delegate for the creation of the list view item's widget */
-	TSharedRef<ITableRow> MakeComponentListItemWidget(TSharedPtr<FMergeComponentData> ComponentData, const TSharedRef<STableViewBase>& OwnerTable);
-
-	/** Delegate to determine whether or not the UI elements should be enabled (determined by number of selected actors / mesh components) */
-	bool GetContentEnabledState() const;
-		
-	/** Delegates for replace actors checkbox UI element*/
-	ECheckBoxState GetReplaceSourceActors() const;
-	void SetReplaceSourceActors(ECheckBoxState NewValue);
-	
-	/** Editor delgates for map and selection changes */
-	void OnLevelSelectionChanged(UObject* Obj);
-	void OnMapChange(uint32 MapFlags);
-	void OnNewCurrentLevel();
-
-	/** Updates SelectedMeshComponent array according to retrieved mesh components from editor selection*/
-	void UpdateSelectedStaticMeshComponents();	
-	/** Stores the individual check box states for the currently selected mesh components */
-	void StoreCheckBoxState();
 private:
 	/** Owning mesh merging tool */
 	FMeshMergingTool* Tool;
 
-
-
-	FComponentSelectionControl ComponentSelectionControl;
-
-	/** Settings view ui element ptr */
-	TSharedPtr<IDetailsView> SettingsView;
 	/** Cached pointer to mesh merging setting singleton object */
 	UMeshMergingSettingsObject* MergeSettings;
-	/** List view state tracking data */
-	bool bRefreshListView;
-
 };
diff --git a/Engine/Source/Editor/MergeActors/Private/MeshProxyTool/MeshProxyTool.cpp b/Engine/Source/Editor/MergeActors/Private/MeshProxyTool/MeshProxyTool.cpp
--- a/Engine/Source/Editor/MergeActors/Private/MeshProxyTool/MeshProxyTool.cpp
+++ b/Engine/Source/Editor/MergeActors/Private/MeshProxyTool/MeshProxyTool.cpp
@@ -11,6 +11,7 @@
 #include "Components/StaticMeshComponent.h"
 #include "Components/InstancedStaticMeshComponent.h"
 #include "Components/SplineMeshComponent.h"
+#include "Engine/StaticMeshActor.h"
 #include "Engine/Selection.h"
 #include "Editor.h"
 #include "MeshProxyTool/SMeshProxyDialog.h"
@@ -22,6 +23,8 @@
 #include "IMeshMergeUtilities.h"
 #include "MeshMergeModule.h"
 
+#include "SMeshProxyDialog.h"
+
 #define LOCTEXT_NAMESPACE "MeshProxyTool"
 
 bool UMeshProxySettingsObject::bInitialized = false;
@@ -46,10 +49,9 @@ TSharedRef<SWidget> FMeshProxyTool::GetWidget()
 
 FText FMeshProxyTool::GetTooltipText() const
 {
-	return LOCTEXT("MeshProxyToolTooltip", "Harvest geometry from selected actors and merge them into single mesh.");
+	return LOCTEXT("MeshProxyToolTooltip", "Merge and simplify to a single static mesh actor");
 }
 
-
 FString FMeshProxyTool::GetDefaultPackageName() const
 {
 	FString PackageName;
@@ -91,35 +93,71 @@ FString FMeshProxyTool::GetDefaultPackageName() const
 	return PackageName;
 }
 
+const TArray<TSharedPtr<FMergeComponentData>>& FMeshProxyTool::GetSelectedComponentsInWidget() const
+{
+	return ProxyDialog->GetSelectedComponents();
+}
+
+void ReplaceSourceActorsByProxyMesh(TArray<UObject*>& NewAssetsToSync, ULevel* Level, TArray<AActor*>& Actors)
+{
+	UStaticMesh* MergedMesh = nullptr;
+	if (NewAssetsToSync.FindItemByClass(&MergedMesh))
+	{
+		Level->Modify();
+
+		UWorld* World = Level->OwningWorld;
+		FActorSpawnParameters Params;
+		Params.OverrideLevel = Level;
+		FRotator MergedActorRotation(ForceInit);
+		// The pivot of the merged mesh is always at the origin
+		FVector MergedActorLocation(0, 0, 0);
+
+		AStaticMeshActor* MergedActor = World->SpawnActor<AStaticMeshActor>(MergedActorLocation, MergedActorRotation, Params);
+		MergedActor->GetStaticMeshComponent()->SetStaticMesh(MergedMesh);
+		MergedActor->SetActorLabel(MergedMesh->GetName());
+		World->UpdateCullDistanceVolumes(MergedActor, MergedActor->GetStaticMeshComponent());
+		GEditor->SelectNone(true, true);
+		GEditor->SelectActor(MergedActor, true, true);
+
+		// Remove source actors
+		for (AActor* Actor : Actors)
+		{
+			Actor->Destroy();
+		}
+	}
+}
 
-bool FMeshProxyTool::RunMerge(const FString& PackageName)
+bool FMeshProxyTool::RunMerge(const FString& PackageName, const TArray<TSharedPtr<FMergeComponentData>>& SelectedComponents)
 {
 	TArray<AActor*> Actors;
+	TArray<ULevel*> UniqueLevels;
 	TArray<UObject*> AssetsToSync;
 
-	// Get the module for the mesh merge utilities
-	const IMeshMergeUtilities& MeshMergeUtilities = FModuleManager::Get().LoadModuleChecked<IMeshMergeModule>("MeshMergeUtilities").GetUtilities();
+	BuildActorsListFromMergeComponentsData(SelectedComponents, Actors, bReplaceSourceActors ? &UniqueLevels : nullptr);
 
-	USelection* SelectedActors = GEditor->GetSelectedActors();
-	for (FSelectionIterator Iter(*SelectedActors); Iter; ++Iter)
+	// This restriction is only for replacement of selected actors with merged mesh actor
+	if (UniqueLevels.Num() > 1 && bReplaceSourceActors)
 	{
-		AActor* Actor = Cast<AActor>(*Iter);
-		if (Actor)
-		{
-			Actors.Add(Actor);
-		}
+		FText Message = NSLOCTEXT("UnrealEd", "FailedToMergeActorsSublevels_Msg", "The selected actors should be in the same level");
+		const FText Title = NSLOCTEXT("UnrealEd", "FailedToMergeActors_Title", "Unable to merge actors");
+		FMessageDialog::Open(EAppMsgType::Ok, Message, &Title);
+		return false;
 	}
 
+	// Get the module for the mesh merge utilities
+	const IMeshMergeUtilities& MeshMergeUtilities = FModuleManager::Get().LoadModuleChecked<IMeshMergeModule>("MeshMergeUtilities").GetUtilities();
+
 	if (Actors.Num())
 	{
 		GWarn->BeginSlowTask(LOCTEXT("MeshProxy_CreatingProxy", "Creating Mesh Proxy"), true);
 		GEditor->BeginTransaction(LOCTEXT("MeshProxy_Create", "Creating Mesh Proxy"));
 
 		FVector ProxyLocation = FVector::ZeroVector;
+		TArray<UObject*> NewAssetsToSync;
 
 		FCreateProxyDelegate ProxyDelegate;
 		ProxyDelegate.BindLambda(
-			[](const FGuid Guid, TArray<UObject*>& InAssetsToSync)
+			[&NewAssetsToSync](const FGuid Guid, TArray<UObject*>& InAssetsToSync)
 		{
 			//Update the asset registry that a new static mash and material has been created
 			if (InAssetsToSync.Num())
@@ -135,11 +173,12 @@ bool FMeshProxyTool::RunMerge(const FString& PackageName)
 				//Also notify the content browser that the new assets exists
 				FContentBrowserModule& ContentBrowserModule = FModuleManager::Get().LoadModuleChecked<FContentBrowserModule>("ContentBrowser");
 				ContentBrowserModule.Get().SyncBrowserToAssets(InAssetsToSync, true);
+
+				NewAssetsToSync += InAssetsToSync;
 			}
 		});
 
 		// Extracting static mesh components from the selected mesh components in the dialog
-		const TArray<TSharedPtr<FMergeComponentData>>& SelectedComponents = ProxyDialog->GetSelectedComponents();
 		TArray<UStaticMeshComponent*> StaticMeshComponentsToMerge;
 
 		for (const TSharedPtr<FMergeComponentData>& SelectedComponent : SelectedComponents)
@@ -159,6 +198,11 @@ bool FMeshProxyTool::RunMerge(const FString& PackageName)
 			MeshMergeUtilities.CreateProxyMesh(StaticMeshComponentsToMerge, SettingsObject->Settings, nullptr, PackageName, JobGuid, ProxyDelegate);
 		}
 
+		if(bReplaceSourceActors && UniqueLevels[0])
+		{
+			ReplaceSourceActorsByProxyMesh(NewAssetsToSync, UniqueLevels[0], Actors);
+		}
+
 		GEditor->EndTransaction();
 		GWarn->EndSlowTask();
 	}
@@ -166,15 +210,10 @@ bool FMeshProxyTool::RunMerge(const FString& PackageName)
 	return true;
 }
 
-bool FMeshProxyTool::CanMerge() const
-{
-	return ProxyDialog->GetNumSelectedMeshComponents() >= 1;
-}
-
 
 FText FThirdPartyMeshProxyTool::GetTooltipText() const
 {
-	return LOCTEXT("ThirdPartyMeshProxyToolTooltip", "Harvest geometry from selected meshes, merge and simplify them as a single mesh.");
+	return LOCTEXT("ThirdPartyMeshProxyToolTooltip", "Merge and simplify meshes into a single mesh.");
 }
 
 
@@ -219,6 +258,31 @@ FString FThirdPartyMeshProxyTool::GetDefaultPackageName() const
 	return PackageName;
 }
 
+bool FThirdPartyMeshProxyTool::RunMergeFromSelection()
+{
+	FString PackageName;
+	if (GetPackageNameForMergeAction(GetDefaultPackageName(), PackageName))
+	{
+		return RunMerge(PackageName);
+	}
+	else
+	{
+		return false;
+	}
+}
+
+bool FThirdPartyMeshProxyTool::RunMergeFromWidget()
+{
+	FString PackageName;
+	if (GetPackageNameForMergeAction(GetDefaultPackageName(), PackageName))
+	{
+		return RunMerge(PackageName);
+	}
+	else
+	{
+		return false;
+	}
+}
 
 bool FThirdPartyMeshProxyTool::RunMerge(const FString& PackageName)
 {
@@ -244,10 +308,11 @@ bool FThirdPartyMeshProxyTool::RunMerge(const FString& PackageName)
 		GEditor->BeginTransaction(LOCTEXT("MeshProxy_Create", "Creating Mesh Proxy"));
 
 		FVector ProxyLocation = FVector::ZeroVector;
+		TArray<UObject*> NewAssetsToSync;
 		
 		FCreateProxyDelegate ProxyDelegate;
 		ProxyDelegate.BindLambda(
-			[](const FGuid Guid, TArray<UObject*>& InAssetsToSync)
+			[&NewAssetsToSync](const FGuid Guid, TArray<UObject*>& InAssetsToSync)
 		{
 			//Update the asset registry that a new static mash and material has been created
 			if (InAssetsToSync.Num())
@@ -263,11 +328,18 @@ bool FThirdPartyMeshProxyTool::RunMerge(const FString& PackageName)
 				//Also notify the content browser that the new assets exists
 				FContentBrowserModule& ContentBrowserModule = FModuleManager::Get().LoadModuleChecked<FContentBrowserModule>("ContentBrowser");
 				ContentBrowserModule.Get().SyncBrowserToAssets(InAssetsToSync, true);
+
+				NewAssetsToSync += InAssetsToSync;
 			}
 		});		
 
 		FGuid JobGuid = FGuid::NewGuid();
-		MeshMergeUtilities.CreateProxyMesh(Actors, ProxySettings, nullptr, PackageName, JobGuid, ProxyDelegate);	
+		MeshMergeUtilities.CreateProxyMesh(Actors, ProxySettings, nullptr, PackageName, JobGuid, ProxyDelegate);
+
+		if (bReplaceSourceActors && Actors[0]->GetLevel())
+		{
+			ReplaceSourceActorsByProxyMesh(NewAssetsToSync, Actors[0]->GetLevel(), Actors);
+		}
 
 		GEditor->EndTransaction();
 		GWarn->EndSlowTask();
@@ -276,7 +348,12 @@ bool FThirdPartyMeshProxyTool::RunMerge(const FString& PackageName)
 	return true;
 }
 
-bool FThirdPartyMeshProxyTool::CanMerge() const
+bool FThirdPartyMeshProxyTool::CanMergeFromSelection() const
+{
+	return true;
+}
+
+bool FThirdPartyMeshProxyTool::CanMergeFromWidget() const
 {
 	return true;
 }
diff --git a/Engine/Source/Editor/MergeActors/Private/MeshProxyTool/MeshProxyTool.h b/Engine/Source/Editor/MergeActors/Private/MeshProxyTool/MeshProxyTool.h
--- a/Engine/Source/Editor/MergeActors/Private/MeshProxyTool/MeshProxyTool.h
+++ b/Engine/Source/Editor/MergeActors/Private/MeshProxyTool/MeshProxyTool.h
@@ -4,11 +4,12 @@
 #include "CoreMinimal.h"
 #include "Widgets/SWidget.h"
 #include "Engine/MeshMerging.h"
-#include "IMergeActorsTool.h"
-#include "SMeshProxyDialog.h"
+#include "MergeActorsTool.h"
 
 #include "MeshProxyTool.generated.h"
 
+class SMeshProxyDialog;
+
 /** Singleton wrapper to allow for using the setting structure in SSettingsView */
 UCLASS(config = Engine)
 class UMeshProxySettingsObject : public UObject
@@ -58,7 +59,7 @@ public:
 /**
 * Mesh Proxy Tool
 */
-class FMeshProxyTool : public IMergeActorsTool
+class FMeshProxyTool : public FMergeActorsTool
 {
 	friend class SMeshProxyDialog;
 
@@ -71,8 +72,10 @@ public:
 	virtual FName GetIconName() const override { return "MergeActors.MeshProxyTool"; }
 	virtual FText GetTooltipText() const override;
 	virtual FString GetDefaultPackageName() const override;
-	virtual bool CanMerge() const override;
-	virtual bool RunMerge(const FString& PackageName) override;
+
+protected:
+	virtual bool RunMerge(const FString& PackageName, const TArray<TSharedPtr<FMergeComponentData>>& SelectedComponents) override;
+	virtual const TArray<TSharedPtr<FMergeComponentData>>& GetSelectedComponentsInWidget() const override;
 
 protected:
 
@@ -97,9 +100,17 @@ public:
 	virtual FName GetIconName() const override { return "MergeActors.MeshProxyTool"; }
 	virtual FText GetTooltipText() const override;
 	virtual FString GetDefaultPackageName() const override;
-	virtual bool RunMerge(const FString& PackageName) override;
-	virtual bool CanMerge() const override;
+	virtual bool GetReplaceSourceActors() const override { return bReplaceSourceActors; }
+	virtual void SetReplaceSourceActors(bool bInReplaceSourceActors) override { bReplaceSourceActors = bInReplaceSourceActors; }
+	virtual bool RunMergeFromSelection() override;
+	virtual bool RunMergeFromWidget() override;
+	virtual bool CanMergeFromSelection() const override;
+	virtual bool CanMergeFromWidget() const override;
+
 protected:
+	virtual bool RunMerge(const FString& PackageName);
 
+protected:
+	bool bReplaceSourceActors = false;
 	FMeshProxySettings ProxySettings;
 };
diff --git a/Engine/Source/Editor/MergeActors/Private/MeshProxyTool/SMeshProxyDialog.cpp b/Engine/Source/Editor/MergeActors/Private/MeshProxyTool/SMeshProxyDialog.cpp
--- a/Engine/Source/Editor/MergeActors/Private/MeshProxyTool/SMeshProxyDialog.cpp
+++ b/Engine/Source/Editor/MergeActors/Private/MeshProxyTool/SMeshProxyDialog.cpp
@@ -24,18 +24,17 @@
 
 #define LOCTEXT_NAMESPACE "SMeshProxyDialog"
 
+//////////////////////////////////////////////////////////////////////////
+// SMeshProxyDialog
 SMeshProxyDialog::SMeshProxyDialog()
 {
-	bRefreshListView = false;
+    MergeStaticMeshComponentsLabel = LOCTEXT("CreateProxyMeshComponentsLabel", "Mesh components used to compute the proxy mesh:");
+	SelectedComponentsListBoxToolTip = LOCTEXT("CreateProxyMeshSelectedComponentsListBoxToolTip", "The selected mesh components will be used to compute the proxy mesh");
+    DeleteUndoLabel = LOCTEXT("DeleteUndo", "Insufficient mesh components found for ProxyLOD merging.");
 }
 
 SMeshProxyDialog::~SMeshProxyDialog()
 {
-	// Remove all delegates
-	USelection::SelectionChangedEvent.RemoveAll(this);
-	USelection::SelectObjectEvent.RemoveAll(this);
-	FEditorDelegates::MapChange.RemoveAll(this);
-	FEditorDelegates::NewCurrentLevel.RemoveAll(this);
 }
 
 BEGIN_SLATE_FUNCTION_BUILD_OPTIMIZATION
@@ -44,178 +43,14 @@ void  SMeshProxyDialog::Construct(const FArguments& InArgs, FMeshProxyTool* InTo
 	checkf(InTool != nullptr, TEXT("Invalid owner tool supplied"));
 	Tool = InTool;
 
-	UpdateSelectedStaticMeshComponents();
-	CreateSettingsView();
-
-	// Create widget layout
-	this->ChildSlot
-	[
-		SNew(SVerticalBox)
-		+ SVerticalBox::Slot()
-		.AutoHeight()
-		.Padding(0, 10, 0, 0)
-		[
-			SNew(SBorder)
-			.BorderImage(FEditorStyle::GetBrush("ToolPanel.GroupBorder"))
-			[
-				SNew(SVerticalBox)
-				// Static mesh component selection
-				+ SVerticalBox::Slot()
-				.AutoHeight()
-				.Padding(FEditorStyle::GetMargin("StandardDialog.ContentPadding"))
-				[
-					SNew(SHorizontalBox)
-					+ SHorizontalBox::Slot()
-					.AutoWidth()
-					.VAlign(VAlign_Center)
-					[
-						SNew(STextBlock)
-						.Text(LOCTEXT("CreateProxyMeshComponentsLabel", "Mesh components used to compute the proxy mesh:"))
-					]
-				]
-				+ SVerticalBox::Slot()
-				.AutoHeight()
-				.Padding(FEditorStyle::GetMargin("StandardDialog.ContentPadding"))
-				[
-					SAssignNew(ComponentSelectionControl.ComponentsListView, SListView<TSharedPtr<FMergeComponentData>>)
-					.ListItemsSource(&ComponentSelectionControl.SelectedComponents)
-					.OnGenerateRow(this, &SMeshProxyDialog::MakeComponentListItemWidget)
-					.ToolTipText(LOCTEXT("CreateProxyMeshSelectedComponentsListBoxToolTip", "The selected mesh components will be used to compute the proxy mesh"))
-				]
-			]
-		]
-
-		+ SVerticalBox::Slot()
-		.Padding(0, 10, 0, 0)
-		[
-			SNew(SBorder)
-			.BorderImage(FEditorStyle::GetBrush("ToolPanel.GroupBorder"))
-			[
-				SNew(SVerticalBox)
-				// Static mesh component selection
-				+ SVerticalBox::Slot()
-				.Padding(FEditorStyle::GetMargin("StandardDialog.ContentPadding"))
-				[
-					SNew(SHorizontalBox)
-					+ SHorizontalBox::Slot()
-					.VAlign(VAlign_Center)
-					[
-						SettingsView->AsShared()
-					]
-				]
-			]
-		]
-
-		+ SVerticalBox::Slot()
-		.AutoHeight()
-		.Padding(10)
-		[
-			SNew(SBorder)
-			.BorderBackgroundColor(FLinearColor::Yellow)
-			.BorderImage(FEditorStyle::GetBrush("ToolPanel.GroupBorder"))
-			.Visibility_Lambda([this]()->EVisibility { return this->GetContentEnabledState() ? EVisibility::Collapsed : EVisibility::Visible; })
-			[
-				SNew(STextBlock)
-				.Text(LOCTEXT("DeleteUndo", "Insufficient mesh components found for ProxyLOD merging."))
-			]
-		]
-	];
-
-
-	// Selection change
-	USelection::SelectionChangedEvent.AddRaw(this, &SMeshProxyDialog::OnLevelSelectionChanged);
-	USelection::SelectObjectEvent.AddRaw(this, &SMeshProxyDialog::OnLevelSelectionChanged);
-	FEditorDelegates::MapChange.AddSP(this, &SMeshProxyDialog::OnMapChange);
-	FEditorDelegates::NewCurrentLevel.AddSP(this, &SMeshProxyDialog::OnNewCurrentLevel);
+	SMeshProxyCommonDialog::Construct(SMeshProxyCommonDialog::FArguments());
 
 	ProxySettings = UMeshProxySettingsObject::Get();
 	SettingsView->SetObject(ProxySettings);
 }
 
-void  SMeshProxyDialog::OnMapChange(uint32 MapFlags)
-{
-	Reset();
-}
-
-void  SMeshProxyDialog::OnNewCurrentLevel()
-{
-	Reset();
-}
-
-void  SMeshProxyDialog::Tick(const FGeometry& AllottedGeometry, const double InCurrentTime, const float InDeltaTime)
-{
-	// Check if we need to update selected components and the listbox
-	if (bRefreshListView == true)
-	{
-		ComponentSelectionControl.UpdateSelectedCompnentsAndListBox();
-
-		bRefreshListView = false;
-	}
-}
-
-
-void SMeshProxyDialog::Reset()
-{
-	bRefreshListView = true;
-}
-
-
-bool SMeshProxyDialog::GetContentEnabledState() const
-{
-	return (GetNumSelectedMeshComponents() >= 1); // Only enabled if a mesh is selected
-}
-
-void SMeshProxyDialog::UpdateSelectedStaticMeshComponents()
-{
-
-	ComponentSelectionControl.UpdateSelectedStaticMeshComponents();
-}
-
-
-TSharedRef<ITableRow> SMeshProxyDialog::MakeComponentListItemWidget(TSharedPtr<FMergeComponentData> ComponentData, const TSharedRef<STableViewBase>& OwnerTable)
-{
-
-	return ComponentSelectionControl.MakeComponentListItemWidget(ComponentData, OwnerTable);
-}
-
-
-void SMeshProxyDialog::CreateSettingsView()
-{
-	// Create a property view
-	FPropertyEditorModule& EditModule = FModuleManager::Get().GetModuleChecked<FPropertyEditorModule>("PropertyEditor");
-
-	FDetailsViewArgs DetailsViewArgs;
-	DetailsViewArgs.bUpdatesFromSelection = true;
-	DetailsViewArgs.bLockable = true;
-	DetailsViewArgs.NameAreaSettings = FDetailsViewArgs::ComponentsAndActorsUseNameArea;
-	DetailsViewArgs.bCustomNameAreaLocation = false;
-	DetailsViewArgs.bCustomFilterAreaLocation = true;
-	DetailsViewArgs.DefaultsOnlyVisibility = EEditDefaultsOnlyNodeVisibility::Hide;
-	DetailsViewArgs.NameAreaSettings = FDetailsViewArgs::ENameAreaSettings::HideNameArea;
-
-
-	// Tiny hack to hide this setting, since we have no way / value to go off to 
-	struct Local
-	{
-		/** Delegate to show all properties */
-		static bool IsPropertyVisible(const FPropertyAndParent& PropertyAndParent, bool bInShouldShowNonEditable)
-		{
-			return (PropertyAndParent.Property.GetFName() != GET_MEMBER_NAME_CHECKED(FMaterialProxySettings, GutterSpace));
-		}
-	};
-
-	SettingsView = EditModule.CreateDetailView(DetailsViewArgs);
-	SettingsView->SetIsPropertyVisibleDelegate(FIsPropertyVisible::CreateStatic(&Local::IsPropertyVisible, true));
-}
-
 END_SLATE_FUNCTION_BUILD_OPTIMIZATION
 
-void SMeshProxyDialog::OnLevelSelectionChanged(UObject* Obj)
-{
-	Reset();
-}
-
-
 void SThirdPartyMeshProxyDialog::Construct(const FArguments& InArgs, FThirdPartyMeshProxyTool* InTool)
 {
 	Tool = InTool;
diff --git a/Engine/Source/Editor/MergeActors/Private/MeshProxyTool/SMeshProxyDialog.h b/Engine/Source/Editor/MergeActors/Private/MeshProxyTool/SMeshProxyDialog.h
--- a/Engine/Source/Editor/MergeActors/Private/MeshProxyTool/SMeshProxyDialog.h
+++ b/Engine/Source/Editor/MergeActors/Private/MeshProxyTool/SMeshProxyDialog.h
@@ -8,9 +8,9 @@
 #include "Widgets/Views/STableViewBase.h"
 
 #include "MergeProxyUtils/Utils.h"
+#include "MergeProxyUtils/SMeshProxyCommonDialog.h"
 
 class FMeshProxyTool;
-class IDetailsView;
 class UMeshProxySettingsObject;
 class UObject;
 
@@ -18,7 +18,7 @@ class UObject;
 SMeshProxyDialog  
 -----------------------------------------------------------------------------*/
 
-class SMeshProxyDialog : public SCompoundWidget
+class SMeshProxyDialog : public SMeshProxyCommonDialog
 {
 public:
 	SLATE_BEGIN_ARGS(SMeshProxyDialog)
@@ -34,51 +34,12 @@ public:
 	/** SWidget functions */
 	void Construct(const FArguments& InArgs, FMeshProxyTool* InTool);
 
-	/** Resets the state of the UI and flags it for refreshing */
-	void Reset();
-
-	/** Getter functionality */
-	const TArray<TSharedPtr<FMergeComponentData>>& GetSelectedComponents() const { return ComponentSelectionControl.SelectedComponents; }
-	/** Get number of selected meshes */
-	const int32 GetNumSelectedMeshComponents() const { return ComponentSelectionControl.NumSelectedMeshComponents; }
-
-	/** Begin override SCompoundWidget */
-	virtual void Tick(const FGeometry& AllottedGeometry, const double InCurrentTime, const float InDeltaTime) override;
-	/** End override SCompoundWidget */
-
-	/** Creates and sets up the settings view element*/
-	void CreateSettingsView();
-
-	/** Delegate for the creation of the list view item's widget */
-	TSharedRef<ITableRow> MakeComponentListItemWidget(TSharedPtr<FMergeComponentData> ComponentData, const TSharedRef<STableViewBase>& OwnerTable);
-
-	/** Delegate to determine whether or not the UI elements should be enabled (determined by number of selected actors / mesh components) */
-	bool GetContentEnabledState() const;
-
-
-	/** Editor delgates for map and selection changes */
-	void OnLevelSelectionChanged(UObject* Obj);
-	void OnMapChange(uint32 MapFlags);
-	void OnNewCurrentLevel();
-
-	/** Updates SelectedMeshComponent array according to retrieved mesh components from editor selection*/
-	void UpdateSelectedStaticMeshComponents();
-	/** Stores the individual check box states for the currently selected mesh components */
-	void StoreCheckBoxState();
 private:
 	/** Owning mesh merging tool */
 	FMeshProxyTool* Tool;
 
-	FComponentSelectionControl ComponentSelectionControl;
-
-	/** Settings view ui element ptr */
-	TSharedPtr<IDetailsView> SettingsView;
-
 	/** Cached pointer to mesh merging setting singleton object */
 	UMeshProxySettingsObject* ProxySettings;
-
-	/** List view state tracking data */
-	bool bRefreshListView;
 };
 
 
diff --git a/Engine/Source/Editor/MergeActors/Private/SMergeActorsToolbar.cpp b/Engine/Source/Editor/MergeActors/Private/SMergeActorsToolbar.cpp
--- a/Engine/Source/Editor/MergeActors/Private/SMergeActorsToolbar.cpp
+++ b/Engine/Source/Editor/MergeActors/Private/SMergeActorsToolbar.cpp
@@ -3,7 +3,6 @@
 #include "SMergeActorsToolbar.h"
 #include "IMergeActorsTool.h"
 #include "Modules/ModuleManager.h"
-#include "Misc/PackageName.h"
 #include "Widgets/SBoxPanel.h"
 #include "Styling/SlateTypes.h"
 #include "Widgets/Layout/SBorder.h"
@@ -12,13 +11,12 @@
 #include "Widgets/Input/SButton.h"
 #include "Widgets/Layout/SScrollBox.h"
 #include "Widgets/Input/SCheckBox.h"
+#include "Widgets/Input/STextComboBox.h"
 #include "EditorStyleSet.h"
 #include "Editor/UnrealEdEngine.h"
 #include "UnrealEdGlobals.h"
 #include "LevelEditor.h"
 #include "IDocumentation.h"
-#include "IContentBrowserSingleton.h"
-#include "ContentBrowserModule.h"
 #include "Styling/ToolBarStyle.h"
 
 #define LOCTEXT_NAMESPACE "SMergeActorsToolbar"
@@ -42,45 +40,66 @@ void SMergeActorsToolbar::Construct(const FArguments& InArgs)
 
 		+ SVerticalBox::Slot()
 		.AutoHeight()
-		.HAlign(HAlign_Left)
 		.Padding(0, 0, 0, 0)
 		[
 			SAssignNew(ToolbarContainer, SBorder)
-			.BorderImage(FEditorStyle::GetBrush("NoBorder"))
-			.Padding(FMargin(4, 0, 0, 0))
+			.BorderImage(FEditorStyle::GetBrush("ToolPanel.GroupBorder"))
+			.Padding(10)
 		]
 
 		+ SVerticalBox::Slot()
 		.FillHeight(1.0f)
-		.Padding(2, 0, 0, 0)
+		.Padding(0, 0, 0, 0)
 		[
 			SNew(SBorder)
-			.BorderImage(FEditorStyle::GetBrush("ToolPanel.GroupBorder"))
+			.BorderImage(FStyleDefaults::GetNoBrush())
 			.Padding(0.f)
 			.IsEnabled(this, &SMergeActorsToolbar::GetContentEnabledState)
 			[
-				SNew(SVerticalBox)
+				SNew(SScrollBox)
+				+SScrollBox::Slot()
+				[
+					SAssignNew(InlineContentHolder, SBox)
+				]
+			]
+		]
 
-				+ SVerticalBox::Slot()
-				.FillHeight(1.0f)
-				.Padding(4, 4, 4, 4)
+		+ SVerticalBox::Slot()
+		.AutoHeight()
+		.VAlign(VAlign_Bottom)
+		[
+			SNew(SBorder)
+			.BorderImage(FEditorStyle::GetBrush("ToolPanel.GroupBorder"))
+			.Padding(10)
+			[
+				SNew(SHorizontalBox) 
+
+				+ SHorizontalBox::Slot()
+				.HAlign(HAlign_Left)
+				.Padding(FEditorStyle::GetMargin("StandardDialog.ContentPadding"))
 				[
-					SNew(SScrollBox)
-					+SScrollBox::Slot()
+					SNew(SCheckBox)
+					.Type(ESlateCheckBoxType::CheckBox)
+					.IsChecked_Lambda([this]() { return this->RegisteredTools[CurrentlySelectedTool]->GetReplaceSourceActors() ? ECheckBoxState::Checked : ECheckBoxState::Unchecked; })
+					.OnCheckStateChanged_Lambda([this](ECheckBoxState NewValue) { this->RegisteredTools[CurrentlySelectedTool]->SetReplaceSourceActors(NewValue == ECheckBoxState::Checked); })
+					.Content()
 					[
-						SAssignNew(InlineContentHolder, SBox)
+						SNew(STextBlock)
+						.Text(LOCTEXT("ReplaceSourceActorsLabel", "Replace Source Actors"))
+						.Font(FEditorStyle::GetFontStyle("StandardDialog.SmallFont"))
 					]
 				]
 
-				+ SVerticalBox::Slot()
-				.AutoHeight()
+				+ SHorizontalBox::Slot()
 				.HAlign(HAlign_Right)
 				.Padding(4, 4, 10, 4)
 				[
 					SNew(SButton)
+					.ButtonStyle(FEditorStyle::Get(), "PrimaryButton")
+					.TextStyle(FEditorStyle::Get(), "DialogButtonText")
 					.Text(LOCTEXT("MergeActors", "Merge Actors"))
 					.OnClicked(this, &SMergeActorsToolbar::OnMergeActorsClicked)
-					.IsEnabled_Lambda([this]() -> bool { return this->RegisteredTools[CurrentlySelectedTool]->CanMerge();})
+					.IsEnabled_Lambda([this]() -> bool { return this->RegisteredTools[CurrentlySelectedTool]->CanMergeFromWidget(); })
 				]
 			]
 		]
@@ -106,54 +125,21 @@ void SMergeActorsToolbar::OnActorSelectionChanged(const TArray<UObject*>& NewSel
 	bIsContentEnabled = (NewSelection.Num() > 0);
 }
 
-
-void SMergeActorsToolbar::OnToolSelectionChanged(const ECheckBoxState NewCheckedState, int32 ToolIndex)
+void SMergeActorsToolbar::OnToolSelectionChanged(TSharedPtr<FDropDownItem> NewSelection, ESelectInfo::Type SelectInfo)
 {
-	if (NewCheckedState == ECheckBoxState::Checked)
+	int32 Index = 0;
+	if (ToolDropDownEntries.Find(NewSelection, Index) && Index != CurrentlySelectedTool)
 	{
-		CurrentlySelectedTool = ToolIndex;
+		CurrentlySelectedTool = Index;
 		UpdateInlineContent();
 	}
 }
 
-
-ECheckBoxState SMergeActorsToolbar::OnIsToolSelected(int32 ToolIndex) const
-{
-	return (CurrentlySelectedTool == ToolIndex) ? ECheckBoxState::Checked : ECheckBoxState::Unchecked;
-}
-
-
 FReply SMergeActorsToolbar::OnMergeActorsClicked()
 {
 	if (CurrentlySelectedTool >= 0 && CurrentlySelectedTool < RegisteredTools.Num())
 	{
-		const FString DefaultPackageName = RegisteredTools[CurrentlySelectedTool]->GetDefaultPackageName();
-		if(DefaultPackageName.Len() > 0)
-		{
-			const FString DefaultPath = FPackageName::GetLongPackagePath(DefaultPackageName);
-			const FString DefaultName = FPackageName::GetShortName(DefaultPackageName);
-
-			// Initialize SaveAssetDialog config
-			FSaveAssetDialogConfig SaveAssetDialogConfig;
-			SaveAssetDialogConfig.DialogTitleOverride = LOCTEXT("CreateMergedActorTitle", "Create Merged Actor");
-			SaveAssetDialogConfig.DefaultPath = DefaultPath;
-			SaveAssetDialogConfig.DefaultAssetName = DefaultName;
-			SaveAssetDialogConfig.ExistingAssetPolicy = ESaveAssetDialogExistingAssetPolicy::AllowButWarn;
-			SaveAssetDialogConfig.AssetClassNames = { UStaticMesh::StaticClass()->GetFName() };
-
-			FContentBrowserModule& ContentBrowserModule = FModuleManager::LoadModuleChecked<FContentBrowserModule>("ContentBrowser");
-			FString SaveObjectPath = ContentBrowserModule.Get().CreateModalSaveAssetDialog(SaveAssetDialogConfig);
-			if (!SaveObjectPath.IsEmpty())
-			{
-				const FString PackageName = FPackageName::ObjectPathToPackageName(SaveObjectPath);
-
-				RegisteredTools[CurrentlySelectedTool]->RunMerge(PackageName);
-			}
-		}
-		else
-		{
-			RegisteredTools[CurrentlySelectedTool]->RunMerge(DefaultPackageName);
-		}
+		RegisteredTools[CurrentlySelectedTool]->RunMergeFromWidget();
 	}
 
 	return FReply::Handled();
@@ -189,43 +175,75 @@ void SMergeActorsToolbar::RemoveTool(IMergeActorsTool* Tool)
 	}
 }
 
+SMergeActorsToolbar::FDropDownItem::FDropDownItem(const FText& InName, const FName& InIconName)
+	: Name(InName), IconName(InIconName)
+{}
 
 void SMergeActorsToolbar::UpdateToolbar()
 {
-	const ISlateStyle& StyleSet = FEditorStyle::Get();
-
-	TSharedRef<SHorizontalBox> HorizontalBox =
-		SNew(SHorizontalBox);
-
-	const FToolBarStyle& BaseToolBarStyle = StyleSet.GetWidgetStyle<FToolBarStyle>("ToolBar");
+	// Build tools entries data for the drop-down
+	ToolDropDownEntries.Empty();
 
-	for (int32 ToolIndex = 0; ToolIndex < RegisteredTools.Num(); ToolIndex++)
+	for(int32 ToolIndex = 0; ToolIndex < RegisteredTools.Num(); ++ToolIndex)
 	{
 		const IMergeActorsTool* Tool = RegisteredTools[ToolIndex];
+		ToolDropDownEntries.Add(MakeShareable(new FDropDownItem(Tool->GetTooltipText(), Tool->GetIconName())));
+	}
 
-		HorizontalBox->AddSlot()
-		.Padding(BaseToolBarStyle.ButtonPadding)
-		.AutoWidth()
+	// Build combo box
+	const ISlateStyle& StyleSet = FEditorStyle::Get();
+
+	TSharedRef <SComboBox<TSharedPtr<FDropDownItem> > > ComboBox =
+		SNew(SComboBox<TSharedPtr<FDropDownItem> >)
+		.OptionsSource(&ToolDropDownEntries)
+		.OnGenerateWidget(this, &SMergeActorsToolbar::MakeWidgetFromEntry)
+		.InitiallySelectedItem(ToolDropDownEntries[CurrentlySelectedTool])
+		.OnSelectionChanged(this, &SMergeActorsToolbar::OnToolSelectionChanged)
+		.Content()
 		[
-			SNew(SCheckBox)
-			.Style(&BaseToolBarStyle.ToggleButton)
-			.OnCheckStateChanged(this, &SMergeActorsToolbar::OnToolSelectionChanged, ToolIndex)
-			.IsChecked(this, &SMergeActorsToolbar::OnIsToolSelected, ToolIndex)
-			.Padding(BaseToolBarStyle.CheckBoxPadding)
-			.ToolTip(IDocumentation::Get()->CreateToolTip(Tool->GetTooltipText(), nullptr, FString(), FString()))
+			SNew(SHorizontalBox)
+			+ SHorizontalBox::Slot()
+			.AutoWidth()
+			.VAlign(VAlign_Center)
 			[
 				SNew(SImage)
-				.Image(StyleSet.GetBrush(Tool->GetIconName()))
+				.Image_Lambda([this](){
+					return FEditorStyle::Get().GetBrush(ToolDropDownEntries[CurrentlySelectedTool]->IconName);
+				})
+			]
+			+ SHorizontalBox::Slot()
+			.FillWidth(1)
+			.VAlign(VAlign_Center)
+			.Padding(5.0, 0, 0, 0)
+			[
+				SNew(STextBlock)
+				.Text_Lambda([this]() -> FText {
+					return ToolDropDownEntries[CurrentlySelectedTool]->Name;
+				})
 			]
 		];
-	}
 
-	TSharedRef<SBorder> ToolbarContent =
-		SNew(SBorder)
-		.Padding(0)
-		.BorderImage(StyleSet.GetBrush("NoBorder"))
+	TSharedRef<SHorizontalBox> ToolbarContent =
+		SNew(SHorizontalBox)
+		+ SHorizontalBox::Slot()
+		.HAlign(HAlign_Left)
+		.VAlign(VAlign_Center)
+		.AutoWidth()
 		[
-			HorizontalBox
+			SNew(STextBlock)
+			.Text(LOCTEXT("MergeMethodLabel", "Merge Method"))
+		]
+		+ SHorizontalBox::Slot()
+		.VAlign(VAlign_Center)
+		.Padding(10.0, 0, 0, 0)
+		[
+			ComboBox
+		]
+		// Filler so that the combo box is just the right size
+		+ SHorizontalBox::Slot()
+		[
+			SNew(SBorder)
+			.BorderImage(FStyleDefaults::GetNoBrush())
 		];
 
 	ToolbarContainer->SetContent(ToolbarContent);
@@ -233,6 +251,25 @@ void SMergeActorsToolbar::UpdateToolbar()
 	UpdateInlineContent();
 }
 
+TSharedRef<SWidget> SMergeActorsToolbar::MakeWidgetFromEntry(TSharedPtr<FDropDownItem> InItem)
+{
+	return SNew(SHorizontalBox)
+		+ SHorizontalBox::Slot()
+		.AutoWidth()
+		.VAlign(VAlign_Center)
+		[
+			SNew(SImage)
+			.Image(FEditorStyle::Get().GetBrush(InItem->IconName))
+		]
+		+ SHorizontalBox::Slot()
+		.FillWidth(1)
+		.VAlign(VAlign_Center)
+		.Padding(5.0, 0, 0, 0)
+		[
+			SNew(STextBlock)
+			.Text(InItem->Name)
+		];
+}
 
 void SMergeActorsToolbar::UpdateInlineContent()
 {
diff --git a/Engine/Source/Editor/MergeActors/Private/SMergeActorsToolbar.h b/Engine/Source/Editor/MergeActors/Private/SMergeActorsToolbar.h
--- a/Engine/Source/Editor/MergeActors/Private/SMergeActorsToolbar.h
+++ b/Engine/Source/Editor/MergeActors/Private/SMergeActorsToolbar.h
@@ -46,15 +46,22 @@ public:
 
 
 private:
+	struct FDropDownItem
+	{
+		FDropDownItem(const FText& InName, const FName& InIconName);
+
+		FText Name;
+		FName IconName;
+	};
+
+	/** Called to create the combo box entries */
+	TSharedRef<SWidget> MakeWidgetFromEntry(TSharedPtr<FDropDownItem> InItem);
 
 	/** Called when the level actor selection changes */
 	void OnActorSelectionChanged(const TArray<UObject*>& NewSelection, bool bForceRefresh);
 
 	/** Called when the currently selected tool changes */
-	void OnToolSelectionChanged(const ECheckBoxState NewCheckedState, int32 ToolIndex);
-
-	/** Called to determine whether a toolbox widget is selected or not */
-	ECheckBoxState OnIsToolSelected(int32 ToolIndex) const;
+	void OnToolSelectionChanged(TSharedPtr<FDropDownItem> NewSelection, ESelectInfo::Type SelectInfo);
 
 	/** Called when the Merge Actors button is clicked */
 	FReply OnMergeActorsClicked();
@@ -68,7 +75,6 @@ private:
 	/** Updates the inline content widget for the current tool */
 	void UpdateInlineContent();
 
-
 private:
 
 	/** List of registered tool instances */
@@ -83,6 +89,9 @@ private:
 	/** Whether the merge actors tool panel is enabled or not */
 	bool bIsContentEnabled;
 
+	/** List of register tool names */
+	TArray< TSharedPtr<FDropDownItem> > ToolDropDownEntries;
+
 	/** The container holding the toolbar */
 	TSharedPtr<SBorder> ToolbarContainer;
 
diff --git a/Engine/Source/Editor/MergeActors/Public/IMergeActorsModule.h b/Engine/Source/Editor/MergeActors/Public/IMergeActorsModule.h
--- a/Engine/Source/Editor/MergeActors/Public/IMergeActorsModule.h
+++ b/Engine/Source/Editor/MergeActors/Public/IMergeActorsModule.h
@@ -32,4 +32,9 @@ public:
 	 * Unregister an IMergeActorsTool with the module
 	 */
 	MERGEACTORS_API virtual bool UnregisterMergeActorsTool(IMergeActorsTool* Tool) = 0;
+
+	/**
+	 * Get currently registered list of MergeActors tools
+	 */
+	MERGEACTORS_API virtual void GetRegisteredMergeActorsTools(TArray<IMergeActorsTool*>& OutTools) = 0;
 };
diff --git a/Engine/Source/Editor/MergeActors/Public/IMergeActorsTool.h b/Engine/Source/Editor/MergeActors/Public/IMergeActorsTool.h
--- a/Engine/Source/Editor/MergeActors/Public/IMergeActorsTool.h
+++ b/Engine/Source/Editor/MergeActors/Public/IMergeActorsTool.h
@@ -36,17 +36,40 @@ public:
 	virtual FString GetDefaultPackageName() const = 0;
 
 	/**
-	 * Perform merge operation
+	 * Checks if the Replace Source Actors option is selected
+	 */
+	virtual bool GetReplaceSourceActors() const = 0;
+
+	/**
+	 * Changes the Replace Source Actors option
+	 * 
+	 * @param whether to replace the source actors or not
+	 */
+	virtual void SetReplaceSourceActors(bool bReplaceSourceActors) = 0;
+
+	/**
+	 * Perform merge operation from the current selection
 	 *
-	 * @param	PackageName		Long package name of the asset to create
-	 * @return	true if the merge succeeded
+	 * @return true if the merge succeeded
 	 */
-	virtual bool RunMerge(const FString& PackageName) = 0;
+	virtual bool RunMergeFromSelection() = 0;
+
+	/**
+	 * Perform merge operation from the selection in the Merge Actors panel
+	 */
+	virtual bool RunMergeFromWidget() = 0;
 	
 	/*
-	* Checks if merge operation is valid
+	* Checks if merge operation is valid from the current selection
+	*
+	* @return true if merge can be executed
+	*/
+	virtual bool CanMergeFromSelection() const = 0;
+
+	/*
+	* Checks if merge operation is valid from the current selection
 	*
-	* @return true if merge is valid
+	* @return true if merge can be executed
 	*/
-	virtual bool CanMerge() const = 0;
+	virtual bool CanMergeFromWidget() const = 0;
 };
diff --git a/Engine/Source/Runtime/Engine/Classes/Engine/MeshMerging.h b/Engine/Source/Runtime/Engine/Classes/Engine/MeshMerging.h
--- a/Engine/Source/Runtime/Engine/Classes/Engine/MeshMerging.h
+++ b/Engine/Source/Runtime/Engine/Classes/Engine/MeshMerging.h
@@ -603,17 +603,6 @@ struct FSectionInfo
 	}
 };
 
-/** How to replace instanced */
-UENUM()
-enum class EMeshInstancingReplacementMethod : uint8
-{
-	/** Destructive workflow: remove the original actors when replacing with instanced static meshes */
-	RemoveOriginalActors,
-
-	/** Non-destructive workflow: keep the original actors but hide them and set them to be editor-only */
-	KeepOriginalActorsAsEditorOnly
-};
-
 /** Mesh instance-replacement settings */
 USTRUCT(Blueprintable)
 struct FMeshInstancingSettings
@@ -623,7 +612,6 @@ struct FMeshInstancingSettings
 	FMeshInstancingSettings()
 		: ActorClassToUse(AActor::StaticClass())
 		, InstanceReplacementThreshold(2)
-		, MeshReplacementMethod(EMeshInstancingReplacementMethod::KeepOriginalActorsAsEditorOnly)
 		, bSkipMeshesWithVertexColors(true)
 		, bUseHLODVolumes(true)
 		, ISMComponentToUse(UInstancedStaticMeshComponent::StaticClass())
@@ -637,10 +625,6 @@ struct FMeshInstancingSettings
 	UPROPERTY(BlueprintReadWrite, EditAnywhere, Category="Instancing", meta=(ClampMin=1))
 	int32 InstanceReplacementThreshold;
 
-	/** How to replace the original actors when instancing */
-	UPROPERTY(BlueprintReadWrite, EditAnywhere, Category="Instancing")
-	EMeshInstancingReplacementMethod MeshReplacementMethod;
-
 	/**
 	 * Whether to skip the conversion to an instanced static mesh for meshes with vertex colors.
 	 * Instanced static meshes do not support vertex colors per-instance, so conversion will lose
