Fix Non-unity build

diff --git a/Engine/Source/Editor/HierarchicalLODOutliner/Private/HLODOutliner.h b/Engine/Source/Editor/HierarchicalLODOutliner/Private/HLODOutliner.h
--- a/Engine/Source/Editor/HierarchicalLODOutliner/Private/HLODOutliner.h
+++ b/Engine/Source/Editor/HierarchicalLODOutliner/Private/HLODOutliner.h
@@ -14,6 +14,7 @@
 #include "Widgets/Views/STableRow.h"
 #include "Widgets/Views/STreeView.h"
 #include "Editor/HierarchicalLODOutliner/Private/TreeItemID.h"
+#include "Textures/SlateIcon.h"
 
 class AActor;
 class ALODActor;
diff --git a/Engine/Source/Editor/MergeActors/Private/MeshProxyTool/MeshProxyTool.cpp b/Engine/Source/Editor/MergeActors/Private/MeshProxyTool/MeshProxyTool.cpp
--- a/Engine/Source/Editor/MergeActors/Private/MeshProxyTool/MeshProxyTool.cpp
+++ b/Engine/Source/Editor/MergeActors/Private/MeshProxyTool/MeshProxyTool.cpp
@@ -14,6 +14,7 @@
 #include "Engine/StaticMeshActor.h"
 #include "Engine/Selection.h"
 #include "Editor.h"
+#include "Misc/MessageDialog.h"
 #include "MeshProxyTool/SMeshProxyDialog.h"
 #include "MeshUtilities.h"
 #include "IContentBrowserSingleton.h"
