Fixed non-unity build, was missing include on static mesh component

diff --git a/Engine/Source/Editor/MergeActors/Private/MergeActorsTool.cpp b/Engine/Source/Editor/MergeActors/Private/MergeActorsTool.cpp
--- a/Engine/Source/Editor/MergeActors/Private/MergeActorsTool.cpp
+++ b/Engine/Source/Editor/MergeActors/Private/MergeActorsTool.cpp
@@ -1,6 +1,7 @@
 // Copyright Epic Games, Inc. All Rights Reserved.
 
 #include "MergeActorsTool.h"
+#include "Components/StaticMeshComponent.h"
 
 bool FMergeActorsTool::GetReplaceSourceActors() const
 {
