MergeActors: add new Approximation method based on GeometryProcesingInterfaces implementation
#rnx

diff --git a/Engine/Source/Editor/MergeActors/MergeActors.Build.cs b/Engine/Source/Editor/MergeActors/MergeActors.Build.cs
--- a/Engine/Source/Editor/MergeActors/MergeActors.Build.cs
+++ b/Engine/Source/Editor/MergeActors/MergeActors.Build.cs
@@ -25,7 +25,8 @@ public class MergeActors : ModuleRules
                 "WorkspaceMenuStructure",
                 "MeshReductionInterface",
                 "MeshMergeUtilities",
-            }
+				"GeometryProcessingInterfaces"
+			}
 		);
 
 		PrivateDependencyModuleNames.AddRange(
@@ -50,6 +51,7 @@ public class MergeActors : ModuleRules
                 "MeshUtilities",
                 "MeshMergeUtilities",
                 "MeshReductionInterface",
+				"GeometryProcessingInterfaces"
             }
 		);
 	}
diff --git a/Engine/Source/Editor/MergeActors/Private/MergeActorsModule.cpp b/Engine/Source/Editor/MergeActors/Private/MergeActorsModule.cpp
--- a/Engine/Source/Editor/MergeActors/Private/MergeActorsModule.cpp
+++ b/Engine/Source/Editor/MergeActors/Private/MergeActorsModule.cpp
@@ -14,6 +14,7 @@
 #include "MeshUtilities.h"
 #include "MeshMergingTool/MeshMergingTool.h"
 #include "MeshProxyTool/MeshProxyTool.h"
+#include "MeshApproximationTool/MeshApproximationTool.h"
 #include "Widgets/Docking/SDockTab.h"
 #include "IMeshReductionManagerModule.h"
 #include "MeshInstancingTool/MeshInstancingTool.h"
@@ -131,6 +132,8 @@ void FMergeActorsModule::StartupModule()
 	}
 
 	ensure(RegisterMergeActorsTool(MakeUnique<FMeshInstancingTool>()));
+
+	ensure(RegisterMergeActorsTool(MakeUnique<FMeshApproximationTool>()));
 }
 
 
diff --git a/Engine/Source/Editor/MergeActors/Private/MeshApproximationTool/MeshApproximationTool.cpp b/Engine/Source/Editor/MergeActors/Private/MeshApproximationTool/MeshApproximationTool.cpp
new file mode 100644
--- /dev/null
+++ b/Engine/Source/Editor/MergeActors/Private/MeshApproximationTool/MeshApproximationTool.cpp
@@ -0,0 +1,230 @@
+// Copyright Epic Games, Inc. All Rights Reserved.
+
+#include "MeshApproximationTool/MeshApproximationTool.h"
+#include "Misc/Paths.h"
+#include "Misc/FeedbackContext.h"
+#include "Modules/ModuleManager.h"
+#include "UObject/Package.h"
+#include "Misc/PackageName.h"
+#include "Widgets/DeclarativeSyntaxSupport.h"
+#include "GameFramework/Actor.h"
+#include "Components/StaticMeshComponent.h"
+#include "Components/InstancedStaticMeshComponent.h"
+#include "Components/SplineMeshComponent.h"
+#include "Engine/StaticMeshActor.h"
+#include "Materials/Material.h"
+#include "Engine/Selection.h"
+#include "Editor.h"
+#include "Misc/MessageDialog.h"
+#include "MeshApproximationTool/SMeshApproximationDialog.h"
+#include "MeshUtilities.h"
+#include "IContentBrowserSingleton.h"
+#include "ContentBrowserModule.h"
+#include "AssetRegistryModule.h"
+#include "IMeshReductionInterfaces.h"
+#include "IMeshMergeUtilities.h"
+#include "MeshMergeModule.h"
+
+#include "AssetRegistryModule.h"
+#include "AssetToolsModule.h"
+
+
+#include "IGeometryProcessingInterfacesModule.h"
+#include "GeometryProcessingInterfaces/ApproximateActors.h"
+
+#include "SMeshApproximationDialog.h"
+
+#define LOCTEXT_NAMESPACE "MeshApproximationTool"
+
+bool UMeshApproximationSettingsObject::bInitialized = false;
+UMeshApproximationSettingsObject* UMeshApproximationSettingsObject::DefaultSettings = nullptr;
+
+
+
+FMeshApproximationTool::FMeshApproximationTool()
+{
+	SettingsObject = UMeshApproximationSettingsObject::Get();
+}
+
+FMeshApproximationTool::~FMeshApproximationTool()
+{
+	UMeshApproximationSettingsObject::Destroy();
+	SettingsObject = nullptr;
+}
+
+TSharedRef<SWidget> FMeshApproximationTool::GetWidget()
+{
+	SAssignNew(ProxyDialog, SMeshApproximationDialog, this);
+	return ProxyDialog.ToSharedRef();
+}
+
+FText FMeshApproximationTool::GetTooltipText() const
+{
+	return LOCTEXT("MeshApproximationToolTooltip", "Approximate selected Actors with combined Actors");
+}
+
+FString FMeshApproximationTool::GetDefaultPackageName() const
+{
+	FString PackageName;
+
+	USelection* SelectedActors = GEditor->GetSelectedActors();
+
+	// Iterate through selected actors and find first static mesh asset
+	// Use this static mesh path as destination package name for a merged mesh
+	for (FSelectionIterator Iter(*SelectedActors); Iter; ++Iter)
+	{
+		AActor* Actor = Cast<AActor>(*Iter);
+		if (Actor)
+		{
+			TInlineComponentArray<UStaticMeshComponent*> SMComponets;
+			Actor->GetComponents<UStaticMeshComponent>(SMComponets);
+			for (UStaticMeshComponent* Component : SMComponets)
+			{
+				if (Component->GetStaticMesh())
+				{
+					PackageName = FPackageName::GetLongPackagePath(Component->GetStaticMesh()->GetOutermost()->GetName());
+					PackageName += FString(TEXT("/APPROX_")) + Component->GetStaticMesh()->GetName();
+					break;
+				}
+			}
+		}
+
+		if (!PackageName.IsEmpty())
+		{
+			break;
+		}
+	}
+
+	if (PackageName.IsEmpty())
+	{
+		PackageName = FPackageName::FilenameToLongPackageName(FPaths::ProjectContentDir() + TEXT("APPROX"));
+		PackageName = MakeUniqueObjectName(NULL, UPackage::StaticClass(), *PackageName).ToString();
+	}
+
+	return PackageName;
+}
+
+
+const TArray<TSharedPtr<FMergeComponentData>>& FMeshApproximationTool::GetSelectedComponentsInWidget() const
+{
+	return ProxyDialog->GetSelectedComponents();
+}
+
+
+static void ReplaceSourceActorsByApproximationMeshes(TArray<UObject*>& NewAssetsToSync, ULevel* Level, TArray<AActor*>& Actors)
+{
+	UStaticMesh* MergedMesh = nullptr;
+	if (NewAssetsToSync.FindItemByClass(&MergedMesh))
+	{
+		Level->Modify();
+
+		UWorld* World = Level->OwningWorld;
+		FActorSpawnParameters Params;
+		Params.OverrideLevel = Level;
+		FRotator MergedActorRotation(ForceInit);
+		// The pivot of the merged mesh is always at the origin
+		FVector MergedActorLocation(0, 0, 0);
+
+		AStaticMeshActor* MergedActor = World->SpawnActor<AStaticMeshActor>(MergedActorLocation, MergedActorRotation, Params);
+		MergedActor->GetStaticMeshComponent()->SetStaticMesh(MergedMesh);
+		MergedActor->SetActorLabel(MergedMesh->GetName());
+		World->UpdateCullDistanceVolumes(MergedActor, MergedActor->GetStaticMeshComponent());
+		GEditor->SelectNone(true, true);
+		GEditor->SelectActor(MergedActor, true, true);
+
+		// Remove source actors
+		for (AActor* Actor : Actors)
+		{
+			Actor->Destroy();
+		}
+	}
+}
+
+bool FMeshApproximationTool::RunMerge(const FString& PackageName, const TArray<TSharedPtr<FMergeComponentData>>& SelectedComponents)
+{
+	TArray<AActor*> Actors;
+	TArray<ULevel*> UniqueLevels;
+	TArray<UObject*> AssetsToSync;
+
+	BuildActorsListFromMergeComponentsData(SelectedComponents, Actors, bReplaceSourceActors ? &UniqueLevels : nullptr);
+
+	// This restriction is only for replacement of selected actors with merged mesh actor
+	if (UniqueLevels.Num() > 1 && bReplaceSourceActors)
+	{
+		FText Message = NSLOCTEXT("UnrealEd", "FailedToMergeActorsSublevels_Msg", "The selected actors should be in the same level");
+		const FText Title = NSLOCTEXT("UnrealEd", "FailedToMergeActors_Title", "Unable to merge actors");
+		FMessageDialog::Open(EAppMsgType::Ok, Message, &Title);
+		return false;
+	}
+
+	const FMeshApproximationToolSettings& UseSettings = SettingsObject->Settings;
+
+	IGeometryProcessingInterfacesModule& GeomProcInterfaces = FModuleManager::Get().LoadModuleChecked<IGeometryProcessingInterfacesModule>("GeometryProcessingInterfaces");
+	IGeometryProcessing_ApproximateActors* ApproxActorsAPI = GeomProcInterfaces.GetApproximateActorsImplementation();
+
+	//
+	// Construct options for ApproximateActors operation
+	//
+
+	IGeometryProcessing_ApproximateActors::FOptions Options;
+	Options.BasePackagePath = PackageName;
+	Options.WorldSpaceApproximationAccuracyMeters = UseSettings.ApproximationAccuracy;
+	Options.ClampVoxelDimension = UseSettings.ClampVoxelDimension;
+	Options.WindingThreshold = UseSettings.WindingThreshold;
+	Options.bApplyMorphology = UseSettings.bFillGaps;
+	Options.MorphologyDistanceMeters = UseSettings.GapDistance;
+	Options.FixedTriangleCount = UseSettings.TargetTriCount;
+	if (UseSettings.SimplifyMethod == EMeshApproximationSimplificationPolicy::TrianglesPerArea)
+	{
+		Options.MeshSimplificationPolicy = IGeometryProcessing_ApproximateActors::ESimplificationPolicy::TrianglesPerUnitSqMeter;
+		Options.SimplificationTargetMetric = UseSettings.TrianglesPerM;
+	}
+	else
+	{
+		Options.MeshSimplificationPolicy = IGeometryProcessing_ApproximateActors::ESimplificationPolicy::FixedTriangleCount;
+	}
+
+	Options.TextureImageSize = UseSettings.MaterialSettings.TextureSize.X;
+	Options.AntiAliasMultiSampling = FMath::Max(1, UseSettings.MultiSamplingAA);
+
+	Options.RenderCaptureImageSize = (UseSettings.RenderCaptureResolution == 0) ?
+		Options.TextureImageSize : UseSettings.RenderCaptureResolution;
+	Options.FieldOfViewDegrees = UseSettings.CaptureFieldOfView;
+	Options.NearPlaneDist = UseSettings.NearPlaneDist;
+
+
+	// run actor approximation computation
+	IGeometryProcessing_ApproximateActors::FResults Results;
+	ApproxActorsAPI->ApproximateActors(Actors, Options, Results);
+
+	// Notify Asset Browser about new Assets
+	for (UStaticMesh* StaticMesh : Results.NewMeshAssets)
+	{
+		FAssetRegistryModule::AssetCreated(StaticMesh);
+	}
+	for (UMaterial* Material : Results.NewMaterials)
+	{
+		FAssetRegistryModule::AssetCreated(Material);
+	}
+	for (UTexture2D* Texture : Results.NewTextures)
+	{
+		FAssetRegistryModule::AssetCreated(Texture);
+	}
+
+	if (ensure(Results.NewMeshAssets.Num() == 1))
+	{
+		TArray<UObject*> NewAssetsToSync;
+		NewAssetsToSync.Add(Results.NewMeshAssets[0]);
+		if (bReplaceSourceActors && UniqueLevels[0])
+		{
+			ReplaceSourceActorsByApproximationMeshes(NewAssetsToSync, UniqueLevels[0], Actors);
+		}
+	}
+
+	return true;
+}
+
+
+
+#undef LOCTEXT_NAMESPACE
+
diff --git a/Engine/Source/Editor/MergeActors/Private/MeshApproximationTool/MeshApproximationTool.h b/Engine/Source/Editor/MergeActors/Private/MeshApproximationTool/MeshApproximationTool.h
new file mode 100644
--- /dev/null
+++ b/Engine/Source/Editor/MergeActors/Private/MeshApproximationTool/MeshApproximationTool.h
@@ -0,0 +1,187 @@
+// Copyright Epic Games, Inc. All Rights Reserved.
+#pragma once
+
+#include "CoreMinimal.h"
+#include "Widgets/SWidget.h"
+#include "Engine/MeshMerging.h"
+#include "MergeActorsTool.h"
+
+#include "MeshApproximationTool.generated.h"
+
+class SMeshApproximationDialog;
+
+
+
+UENUM()
+enum class EMeshApproximationSimplificationPolicy : uint8
+{
+	FixedTriangleCount = 0,
+	TrianglesPerArea = 1
+};
+
+
+USTRUCT(Blueprintable)
+struct FMeshApproximationToolSettings
+{
+	GENERATED_BODY()
+
+	/** Approximation Accuracy in Meters, will determine (eg) voxel resolution */
+	UPROPERTY(BlueprintReadWrite, EditAnywhere, Category = ShapeSettings, meta = (ClampMin = "0.1"))
+	float ApproximationAccuracy = 1.0;
+
+	/** Maximum allowable voxel count along main directions. This is a limit on ApproximationAccuracy */
+	UPROPERTY(BlueprintReadWrite, EditAnywhere, AdvancedDisplay, Category = ShapeSettings, meta = (ClampMin = "64"))
+	int32 ClampVoxelDimension = 512;
+
+	/** */
+	UPROPERTY(BlueprintReadWrite, EditAnywhere, Category = ShapeSettings, meta = (ClampMin = "0.01", ClampMax = "0.99"))
+	float WindingThreshold = 0.5;
+
+	UPROPERTY(BlueprintReadWrite, EditAnywhere, Category = ShapeSettings)
+	bool bFillGaps = true;
+
+	/** Distance in Meters to expand approximation to fill gaps */
+	UPROPERTY(BlueprintReadWrite, EditAnywhere, Category = ShapeSettings, meta = (ClampMin = "0.1", EditCondition = "bFillGaps"))
+	float GapDistance = 0.1;
+
+	UPROPERTY(BlueprintReadWrite, EditAnywhere, Category = ShapeSettings)
+	EMeshApproximationSimplificationPolicy SimplifyMethod = EMeshApproximationSimplificationPolicy::FixedTriangleCount;
+
+	UPROPERTY(BlueprintReadWrite, EditAnywhere, Category = ShapeSettings, meta = (ClampMin = "16", EditCondition = "SimplifyMethod == EMeshApproximationSimplificationPolicy::FixedTriangleCount" ))
+	int32 TargetTriCount = 2000;
+
+	/** Approximate Number of triangles per Square Meter */
+	UPROPERTY(BlueprintReadWrite, EditAnywhere, Category = ShapeSettings, meta = (ClampMin = "0.01", EditCondition = "SimplifyMethod == EMeshApproximationSimplificationPolicy::TrianglesPerArea" ))
+	float TrianglesPerM = 2.0f;
+
+
+
+	/** Material simplification */
+	UPROPERTY(BlueprintReadWrite, EditAnywhere, Category = MaterialSettings)
+	FMaterialProxySettings MaterialSettings;
+
+	/** If Value is > 1, Multisample output baked textures by this amount in each direction (eg 4 == 16x supersampling) */
+	UPROPERTY(BlueprintReadWrite, EditAnywhere, AdvancedDisplay, Category = MaterialSettings, meta = (ClampMin = "0", ClampMax = "8", UIMin = "0", UIMax = "4"))
+	int32 MultiSamplingAA = 0;
+
+	UPROPERTY(BlueprintReadWrite, EditAnywhere, AdvancedDisplay, Category = MaterialSettings, meta = (ClampMin = "5.0", ClampMax = "160.0"))
+	float CaptureFieldOfView = 30.0f;
+
+	UPROPERTY(BlueprintReadWrite, EditAnywhere, AdvancedDisplay, Category = MaterialSettings, meta = (ClampMin = "0.001", ClampMax = "1000.0"))
+	float NearPlaneDist = 1.0f;
+
+	/** If Value is zero, use MaterialSettings resolution, otherwise override the render capture resolution */
+	UPROPERTY(BlueprintReadWrite, EditAnywhere, AdvancedDisplay, Category = MaterialSettings, meta = (ClampMin = "128"))
+	int32 RenderCaptureResolution = 0;
+
+	/** Equality operator. */
+	bool operator==(const FMeshApproximationToolSettings& Other) const
+	{
+		return ApproximationAccuracy == Other.ApproximationAccuracy
+			&& ClampVoxelDimension == Other.ClampVoxelDimension
+			&& WindingThreshold == Other.WindingThreshold
+			&& bFillGaps == Other.bFillGaps
+			&& GapDistance == Other.GapDistance
+			&& SimplifyMethod == Other.SimplifyMethod
+			&& TargetTriCount == Other.TargetTriCount
+			&& TrianglesPerM == Other.TrianglesPerM
+
+			&& MaterialSettings == Other.MaterialSettings
+			&& MultiSamplingAA == Other.MultiSamplingAA
+			&& CaptureFieldOfView == Other.CaptureFieldOfView
+			&& NearPlaneDist == Other.NearPlaneDist
+			&& RenderCaptureResolution == Other.RenderCaptureResolution;
+	}
+
+	/** Inequality. */
+	bool operator!=(const FMeshApproximationToolSettings& Other) const
+	{
+		return !(*this == Other);
+	}
+
+#if WITH_EDITORONLY_DATA
+	/** Handles deprecated properties */
+	void PostLoadDeprecated() {}		// none currently
+#endif
+};
+
+
+
+
+
+/** Singleton wrapper to allow for using the setting structure in SSettingsView */
+UCLASS(config = Engine)
+class UMeshApproximationSettingsObject : public UObject
+{
+	GENERATED_BODY()
+public:
+	UMeshApproximationSettingsObject()
+	{
+	}
+
+	static UMeshApproximationSettingsObject* Get()
+	{	
+		// This is a singleton, duplicate default object
+		if (!bInitialized)
+		{
+			DefaultSettings = DuplicateObject(GetMutableDefault<UMeshApproximationSettingsObject>(), nullptr);
+			DefaultSettings->AddToRoot();
+			bInitialized = true;
+		}
+
+		return DefaultSettings;
+	}
+
+	static void Destroy()
+	{
+		if (bInitialized)
+		{
+			if (UObjectInitialized() && DefaultSettings)
+			{
+				DefaultSettings->RemoveFromRoot();
+				DefaultSettings->MarkPendingKill();
+			}
+
+			DefaultSettings = nullptr;
+			bInitialized = false;
+		}
+	}
+
+protected:
+	static bool bInitialized;
+	static UMeshApproximationSettingsObject* DefaultSettings;
+public:
+	UPROPERTY(EditAnywhere, meta = (ShowOnlyInnerProperties), Category = ApproximationSettings)
+	FMeshApproximationToolSettings Settings;
+};
+
+/**
+* Mesh Proxy Tool
+*/
+class FMeshApproximationTool : public FMergeActorsTool
+{
+	friend class SMeshApproximationDialog;
+
+public:
+	FMeshApproximationTool();
+	~FMeshApproximationTool();
+
+	// IMergeActorsTool interface
+	virtual TSharedRef<SWidget> GetWidget() override;
+	virtual FName GetIconName() const override { return "MergeActors.MeshApproximationTool"; }
+	virtual FText GetTooltipText() const override;
+	virtual FString GetDefaultPackageName() const override;
+
+protected:
+	virtual bool RunMerge(const FString& PackageName, const TArray<TSharedPtr<FMergeComponentData>>& SelectedComponents) override;
+	virtual const TArray<TSharedPtr<FMergeComponentData>>& GetSelectedComponentsInWidget() const override;
+
+protected:
+
+	/** Pointer to the mesh merging dialog containing settings for the merge */
+	TSharedPtr<SMeshApproximationDialog> ProxyDialog;
+
+	/** Pointer to singleton settings object */
+	UMeshApproximationSettingsObject* SettingsObject;
+};
+
diff --git a/Engine/Source/Editor/MergeActors/Private/MeshApproximationTool/SMeshApproximationDialog.cpp b/Engine/Source/Editor/MergeActors/Private/MeshApproximationTool/SMeshApproximationDialog.cpp
new file mode 100644
--- /dev/null
+++ b/Engine/Source/Editor/MergeActors/Private/MeshApproximationTool/SMeshApproximationDialog.cpp
@@ -0,0 +1,56 @@
+// Copyright Epic Games, Inc. All Rights Reserved.
+
+#include "MeshApproximationTool/SMeshApproximationDialog.h"
+#include "Editor.h"
+#include "EditorStyleSet.h"
+#include "Engine/MeshMerging.h"
+#include "Engine/Selection.h"
+#include "MeshApproximationTool/MeshApproximationTool.h"
+#include "Modules/ModuleManager.h"
+#include "PropertyEditorModule.h"
+#include "Styling/SlateTypes.h"
+#include "SlateOptMacros.h"
+#include "UObject/UnrealType.h"
+#include "Widgets/Images/SImage.h"
+#include "Widgets/Input/SCheckBox.h"
+#include "Widgets/Input/SNumericEntryBox.h"
+#include "Widgets/Input/STextComboBox.h"
+#include "Widgets/Layout/SBorder.h"
+#include "Widgets/Layout/SBox.h"
+#include "Widgets/SBoxPanel.h"
+#include "Widgets/Text/STextBlock.h"
+#include "Widgets/Views/SListView.h"
+
+
+#define LOCTEXT_NAMESPACE "SMeshApproximationDialog"
+
+//////////////////////////////////////////////////////////////////////////
+// SMeshApproximationDialog
+SMeshApproximationDialog::SMeshApproximationDialog()
+{
+    MergeStaticMeshComponentsLabel = LOCTEXT("CreateProxyMeshComponentsLabel", "Mesh components used to compute the proxy mesh:");
+	SelectedComponentsListBoxToolTip = LOCTEXT("CreateProxyMeshSelectedComponentsListBoxToolTip", "The selected mesh components will be used to compute the proxy mesh");
+    DeleteUndoLabel = LOCTEXT("DeleteUndo", "Insufficient mesh components found for Mesh Approximation.");
+}
+
+SMeshApproximationDialog::~SMeshApproximationDialog()
+{
+}
+
+BEGIN_SLATE_FUNCTION_BUILD_OPTIMIZATION
+void  SMeshApproximationDialog::Construct(const FArguments& InArgs, FMeshApproximationTool* InTool)
+{
+	checkf(InTool != nullptr, TEXT("Invalid owner tool supplied"));
+	Tool = InTool;
+
+	SMeshProxyCommonDialog::Construct(SMeshProxyCommonDialog::FArguments());
+
+	ProxySettings = UMeshApproximationSettingsObject::Get();
+	SettingsView->SetObject(ProxySettings);
+}
+
+END_SLATE_FUNCTION_BUILD_OPTIMIZATION
+
+
+
+#undef LOCTEXT_NAMESPACE
diff --git a/Engine/Source/Editor/MergeActors/Private/MeshApproximationTool/SMeshApproximationDialog.h b/Engine/Source/Editor/MergeActors/Private/MeshApproximationTool/SMeshApproximationDialog.h
new file mode 100644
--- /dev/null
+++ b/Engine/Source/Editor/MergeActors/Private/MeshApproximationTool/SMeshApproximationDialog.h
@@ -0,0 +1,45 @@
+// Copyright Epic Games, Inc. All Rights Reserved.
+
+#pragma once
+
+#include "CoreMinimal.h"
+#include "Widgets/DeclarativeSyntaxSupport.h"
+#include "Widgets/SCompoundWidget.h"
+#include "Widgets/Views/STableViewBase.h"
+
+#include "MergeProxyUtils/Utils.h"
+#include "MergeProxyUtils/SMeshProxyCommonDialog.h"
+
+class FMeshApproximationTool;
+class UMeshApproximationSettingsObject;
+class UObject;
+
+/*-----------------------------------------------------------------------------
+SMeshApproximationDialog  
+-----------------------------------------------------------------------------*/
+
+class SMeshApproximationDialog : public SMeshProxyCommonDialog
+{
+public:
+	SLATE_BEGIN_ARGS(SMeshApproximationDialog)
+	{
+	}
+	SLATE_END_ARGS()
+
+public:
+	/** **/
+	SMeshApproximationDialog();
+	~SMeshApproximationDialog();
+
+	/** SWidget functions */
+	void Construct(const FArguments& InArgs, FMeshApproximationTool* InTool);
+
+private:
+	/** Owning mesh merging tool */
+	FMeshApproximationTool* Tool;
+
+	/** Cached pointer to mesh merging setting singleton object */
+	UMeshApproximationSettingsObject* ProxySettings;
+};
+
+
