Actor Merging Tool
* Fixed copy paste error causing the wrong texture size to be computed when using the approximate actor merge method along with an automatic texture size based on the mesh draw distance

diff --git a/Engine/Source/Editor/MergeActors/Private/MeshApproximationTool/MeshApproximationTool.cpp b/Engine/Source/Editor/MergeActors/Private/MeshApproximationTool/MeshApproximationTool.cpp
--- a/Engine/Source/Editor/MergeActors/Private/MeshApproximationTool/MeshApproximationTool.cpp
+++ b/Engine/Source/Editor/MergeActors/Private/MeshApproximationTool/MeshApproximationTool.cpp
@@ -201,7 +201,7 @@ bool FMeshApproximationTool::RunMerge(const FString& PackageName, const TArray<T
 	}
 	else if (TextureSizingType == ETextureSizingType::TextureSizingType_AutomaticFromMeshDrawDistance)
 	{
-		TexelDensityPerMeter = FMaterialUtilities::ComputeRequiredTexelDensityFromScreenSize(SettingsObject->Settings.MaterialSettings.MeshMinDrawDistance, GetActorsBounds().SphereRadius);
+		TexelDensityPerMeter = FMaterialUtilities::ComputeRequiredTexelDensityFromDrawDistance(SettingsObject->Settings.MaterialSettings.MeshMinDrawDistance, GetActorsBounds().SphereRadius);
 		TextureSizePolicy = IGeometryProcessing_ApproximateActors::ETextureSizePolicy::TexelDensity;
 	}
 
