Doc / UI improvements for the Merge Actors tool
#rnx

diff --git a/Engine/Source/Editor/LevelEditor/Private/LevelEditorContextMenu.cpp b/Engine/Source/Editor/LevelEditor/Private/LevelEditorContextMenu.cpp
--- a/Engine/Source/Editor/LevelEditor/Private/LevelEditorContextMenu.cpp
+++ b/Engine/Source/Editor/LevelEditor/Private/LevelEditorContextMenu.cpp
@@ -1418,7 +1418,7 @@ void FLevelEditorContextMenuImpl::FillMergeActorsMenu(UToolMenu* Menu)
 
 		Section.AddMenuEntry(
 			NAME_None,
-			Tool->GetTooltipText(),
+			Tool->GetToolNameText(),
 			Tool->GetTooltipText(),
 			FSlateIcon(FAppStyle::GetAppStyleSetName(), Tool->GetIconName()),
 			MergeActorToolAction);
@@ -1431,7 +1431,7 @@ void FLevelEditorContextMenuImpl::FillMergeActorsMenu(UToolMenu* Menu)
 
 	Section.AddMenuEntry(
 		FLevelEditorCommands::Get().OpenMergeActor,
-		LOCTEXT("OpenMergeActor", "Merge Actors Settings"),
+		LOCTEXT("OpenMergeActor", "Merge Actors Settings..."),
 		LOCTEXT("OpenMergeActor_ToolTip", "Click to open the Merge Actor panel"),
 		FSlateIcon(FAppStyle::GetAppStyleSetName(), "LevelEditor.GameSettings"));
 }
diff --git a/Engine/Source/Editor/MergeActors/Private/MeshApproximationTool/MeshApproximationTool.cpp b/Engine/Source/Editor/MergeActors/Private/MeshApproximationTool/MeshApproximationTool.cpp
--- a/Engine/Source/Editor/MergeActors/Private/MeshApproximationTool/MeshApproximationTool.cpp
+++ b/Engine/Source/Editor/MergeActors/Private/MeshApproximationTool/MeshApproximationTool.cpp
@@ -59,9 +59,19 @@ TSharedRef<SWidget> FMeshApproximationTool::GetWidget()
 	return ProxyDialog.ToSharedRef();
 }
 
+FName FMeshApproximationTool::GetIconName() const
+{
+	return "MergeActors.Approximate";
+}
+
+FText FMeshApproximationTool::GetToolNameText() const
+{
+	return LOCTEXT("MeshApproximationToolName", "Approximate");
+}
+
 FText FMeshApproximationTool::GetTooltipText() const
 {
-	return LOCTEXT("MeshApproximationToolTooltip", "Approximate selected Actors with combined Actors");
+	return LOCTEXT("MeshApproximationToolTooltip", "Merge source actors meshes and perform an approximation pass. Will generate a single static mesh with baked textures.");
 }
 
 FString FMeshApproximationTool::GetDefaultPackageName() const
diff --git a/Engine/Source/Editor/MergeActors/Private/MeshApproximationTool/MeshApproximationTool.h b/Engine/Source/Editor/MergeActors/Private/MeshApproximationTool/MeshApproximationTool.h
--- a/Engine/Source/Editor/MergeActors/Private/MeshApproximationTool/MeshApproximationTool.h
+++ b/Engine/Source/Editor/MergeActors/Private/MeshApproximationTool/MeshApproximationTool.h
@@ -69,7 +69,8 @@ public:
 
 	// IMergeActorsTool interface
 	virtual TSharedRef<SWidget> GetWidget() override;
-	virtual FName GetIconName() const override { return "MergeActors.Approximate"; }
+	virtual FName GetIconName() const override;
+	virtual FText GetToolNameText() const override;
 	virtual FText GetTooltipText() const override;
 	virtual FString GetDefaultPackageName() const override;
 
diff --git a/Engine/Source/Editor/MergeActors/Private/MeshInstancingTool/MeshInstancingTool.cpp b/Engine/Source/Editor/MergeActors/Private/MeshInstancingTool/MeshInstancingTool.cpp
--- a/Engine/Source/Editor/MergeActors/Private/MeshInstancingTool/MeshInstancingTool.cpp
+++ b/Engine/Source/Editor/MergeActors/Private/MeshInstancingTool/MeshInstancingTool.cpp
@@ -44,9 +44,19 @@ TSharedRef<SWidget> FMeshInstancingTool::GetWidget()
 	return InstancingDialog.ToSharedRef();
 }
 
+FName FMeshInstancingTool::GetIconName() const
+{
+	return "MergeActors.MeshInstancingTool";
+}
+
+FText FMeshInstancingTool::GetToolNameText() const
+{
+	return LOCTEXT("MeshInstancingToolName", "Batch");
+}
+
 FText FMeshInstancingTool::GetTooltipText() const
 {
-	return LOCTEXT("MeshInstancingToolTooltip", "Merge to an actor with instanced static mesh components");
+	return LOCTEXT("MeshInstancingToolTooltip", "Batch the source actors components to use instancing as much as possible. Will generate an actor with instanced static mesh component(s).");
 }
 
 FString FMeshInstancingTool::GetDefaultPackageName() const
diff --git a/Engine/Source/Editor/MergeActors/Private/MeshInstancingTool/MeshInstancingTool.h b/Engine/Source/Editor/MergeActors/Private/MeshInstancingTool/MeshInstancingTool.h
--- a/Engine/Source/Editor/MergeActors/Private/MeshInstancingTool/MeshInstancingTool.h
+++ b/Engine/Source/Editor/MergeActors/Private/MeshInstancingTool/MeshInstancingTool.h
@@ -73,7 +73,8 @@ public:
 
 	// IMergeActorsTool interface
 	virtual TSharedRef<SWidget> GetWidget() override;
-	virtual FName GetIconName() const override { return "MergeActors.MeshInstancingTool"; }
+	virtual FName GetIconName() const override;
+	virtual FText GetToolNameText() const override;
 	virtual FText GetTooltipText() const override;
 	virtual FString GetDefaultPackageName() const override;
 
diff --git a/Engine/Source/Editor/MergeActors/Private/MeshMergingTool/MeshMergingTool.cpp b/Engine/Source/Editor/MergeActors/Private/MeshMergingTool/MeshMergingTool.cpp
--- a/Engine/Source/Editor/MergeActors/Private/MeshMergingTool/MeshMergingTool.cpp
+++ b/Engine/Source/Editor/MergeActors/Private/MeshMergingTool/MeshMergingTool.cpp
@@ -44,9 +44,19 @@ TSharedRef<SWidget> FMeshMergingTool::GetWidget()
 	return MergingDialog.ToSharedRef();
 }
 
+FName FMeshMergingTool::GetIconName() const
+{
+	return "MergeActors.MeshMergingTool";
+}
+
+FText FMeshMergingTool::GetToolNameText() const
+{
+	return LOCTEXT("MeshMergingToolName", "Merge");
+}
+
 FText FMeshMergingTool::GetTooltipText() const
 {
-	return LOCTEXT("MeshMergingToolTooltip", "Merge to a single static mesh actor");
+	return LOCTEXT("MeshMergingToolTooltip", "Merge the source actors components to generate a single mesh. No simplification pass is performed. Will generate a single static mesh & optionally bake down textures.");
 }
 
 FString FMeshMergingTool::GetDefaultPackageName() const
diff --git a/Engine/Source/Editor/MergeActors/Private/MeshMergingTool/MeshMergingTool.h b/Engine/Source/Editor/MergeActors/Private/MeshMergingTool/MeshMergingTool.h
--- a/Engine/Source/Editor/MergeActors/Private/MeshMergingTool/MeshMergingTool.h
+++ b/Engine/Source/Editor/MergeActors/Private/MeshMergingTool/MeshMergingTool.h
@@ -78,7 +78,8 @@ public:
 
 	// IMergeActorsTool interface
 	virtual TSharedRef<SWidget> GetWidget() override;
-	virtual FName GetIconName() const override { return "MergeActors.MeshMergingTool"; }
+	virtual FName GetIconName() const override;
+	virtual FText GetToolNameText() const override;
 	virtual FText GetTooltipText() const override;
 	virtual FString GetDefaultPackageName() const override;
 
diff --git a/Engine/Source/Editor/MergeActors/Private/MeshProxyTool/MeshProxyTool.cpp b/Engine/Source/Editor/MergeActors/Private/MeshProxyTool/MeshProxyTool.cpp
--- a/Engine/Source/Editor/MergeActors/Private/MeshProxyTool/MeshProxyTool.cpp
+++ b/Engine/Source/Editor/MergeActors/Private/MeshProxyTool/MeshProxyTool.cpp
@@ -48,9 +48,19 @@ TSharedRef<SWidget> FMeshProxyTool::GetWidget()
 	return ProxyDialog.ToSharedRef();
 }
 
+FName FMeshProxyTool::GetIconName() const
+{
+	return "MergeActors.MeshProxyTool";
+}
+
+FText FMeshProxyTool::GetToolNameText() const
+{
+	return LOCTEXT("MeshProxyToolName", "Simplify");
+}
+
 FText FMeshProxyTool::GetTooltipText() const
 {
-	return LOCTEXT("MeshProxyToolTooltip", "Merge and simplify to a single static mesh actor");
+	return LOCTEXT("MeshProxyToolTooltip", "Merge source actors meshes and perform a simplification pass. Will generate a single static mesh with baked textures.");
 }
 
 FString FMeshProxyTool::GetDefaultPackageName() const
@@ -211,12 +221,20 @@ bool FMeshProxyTool::RunMerge(const FString& PackageName, const TArray<TSharedPt
 	return true;
 }
 
+FName FThirdPartyMeshProxyTool::GetIconName() const 
+{
+	return "MergeActors.MeshProxyTool";
+}
 
-FText FThirdPartyMeshProxyTool::GetTooltipText() const
+FText FThirdPartyMeshProxyTool::GetToolNameText() const
 {
-	return LOCTEXT("ThirdPartyMeshProxyToolTooltip", "Merge and simplify meshes into a single mesh.");
+	return LOCTEXT("ThirdPartyMeshProxyToolName", "Third-Party Simplification");
 }
 
+FText FThirdPartyMeshProxyTool::GetTooltipText() const
+{
+	return LOCTEXT("ThirdPartyMeshProxyToolTooltip", "Merge source actors meshes and perform a simplification pass. Will generate a single static mesh with baked textures.");
+}
 
 FString FThirdPartyMeshProxyTool::GetDefaultPackageName() const
 {
diff --git a/Engine/Source/Editor/MergeActors/Private/MeshProxyTool/MeshProxyTool.h b/Engine/Source/Editor/MergeActors/Private/MeshProxyTool/MeshProxyTool.h
--- a/Engine/Source/Editor/MergeActors/Private/MeshProxyTool/MeshProxyTool.h
+++ b/Engine/Source/Editor/MergeActors/Private/MeshProxyTool/MeshProxyTool.h
@@ -69,7 +69,8 @@ public:
 
 	// IMergeActorsTool interface
 	virtual TSharedRef<SWidget> GetWidget() override;
-	virtual FName GetIconName() const override { return "MergeActors.MeshProxyTool"; }
+	virtual FName GetIconName() const override;
+	virtual FText GetToolNameText() const override;
 	virtual FText GetTooltipText() const override;
 	virtual FString GetDefaultPackageName() const override;
 
@@ -97,7 +98,8 @@ public:
 
 	// IMergeActorsTool interface
 	virtual TSharedRef<SWidget> GetWidget() override;
-	virtual FName GetIconName() const override { return "MergeActors.MeshProxyTool"; }
+	virtual FName GetIconName() const override;
+	virtual FText GetToolNameText() const override;
 	virtual FText GetTooltipText() const override;
 	virtual FString GetDefaultPackageName() const override;
 	virtual bool GetReplaceSourceActors() const override { return bReplaceSourceActors; }
diff --git a/Engine/Source/Editor/MergeActors/Private/SMergeActorsToolbar.cpp b/Engine/Source/Editor/MergeActors/Private/SMergeActorsToolbar.cpp
--- a/Engine/Source/Editor/MergeActors/Private/SMergeActorsToolbar.cpp
+++ b/Engine/Source/Editor/MergeActors/Private/SMergeActorsToolbar.cpp
@@ -174,8 +174,10 @@ void SMergeActorsToolbar::RemoveTool(IMergeActorsTool* Tool)
 	}
 }
 
-SMergeActorsToolbar::FDropDownItem::FDropDownItem(const FText& InName, const FName& InIconName)
-	: Name(InName), IconName(InIconName)
+SMergeActorsToolbar::FDropDownItem::FDropDownItem(const FText& InName, const FName& InIconName, const FText& InDescription)
+	: Name(InName)
+	, IconName(InIconName)
+	, Description(InDescription)
 {}
 
 void SMergeActorsToolbar::UpdateToolbar()
@@ -186,7 +188,7 @@ void SMergeActorsToolbar::UpdateToolbar()
 	for(int32 ToolIndex = 0; ToolIndex < RegisteredTools.Num(); ++ToolIndex)
 	{
 		const IMergeActorsTool* Tool = RegisteredTools[ToolIndex];
-		ToolDropDownEntries.Add(MakeShareable(new FDropDownItem(Tool->GetTooltipText(), Tool->GetIconName())));
+		ToolDropDownEntries.Add(MakeShareable(new FDropDownItem(Tool->GetToolNameText(), Tool->GetIconName(), Tool->GetTooltipText())));
 	}
 
 	// Build combo box
@@ -222,27 +224,56 @@ void SMergeActorsToolbar::UpdateToolbar()
 			]
 		];
 
-	TSharedRef<SHorizontalBox> ToolbarContent =
-		SNew(SHorizontalBox)
-		+ SHorizontalBox::Slot()
-		.HAlign(HAlign_Left)
-		.VAlign(VAlign_Center)
-		.AutoWidth()
-		[
-			SNew(STextBlock)
-			.Text(LOCTEXT("MergeMethodLabel", "Merge Method"))
-		]
-		+ SHorizontalBox::Slot()
-		.VAlign(VAlign_Center)
-		.Padding(10.0, 0, 0, 0)
+	TSharedRef<SVerticalBox> ToolbarContent =
+		SNew(SVerticalBox)
+
+		+ SVerticalBox::Slot()
+		.AutoHeight()
 		[
-			ComboBox
+			SNew(SHorizontalBox)
+			+ SHorizontalBox::Slot()
+			.HAlign(HAlign_Left)
+			.VAlign(VAlign_Center)
+			.AutoWidth()
+			[
+				SNew(STextBlock)
+				.Text(LOCTEXT("MergeMethodLabel", "Merge Method"))
+			]
+			+ SHorizontalBox::Slot()
+			.VAlign(VAlign_Center)
+			.Padding(10.0, 0, 0, 0)
+			[
+				ComboBox
+			]
+			// Filler so that the combo box is just the right size
+			+ SHorizontalBox::Slot()
+			[
+				SNew(SBorder)
+				.BorderImage(FStyleDefaults::GetNoBrush())
+			]
 		]
-		// Filler so that the combo box is just the right size
-		+ SHorizontalBox::Slot()
+		+ SVerticalBox::Slot()
+		.Padding(10, 10, 10, 10)
+		.AutoHeight()
 		[
-			SNew(SBorder)
-			.BorderImage(FStyleDefaults::GetNoBrush())
+			SNew(SHorizontalBox)
+			+ SHorizontalBox::Slot()
+			.AutoWidth()
+			.HAlign(HAlign_Center)
+			.VAlign(VAlign_Center)
+			[
+				SNew(SImage)
+				.Image(FEditorStyle::Get().GetBrush("Icons.Info"))
+			]
+			+ SHorizontalBox::Slot()
+			.Padding(10, 0, 0, 0)
+			[
+				SNew(STextBlock)
+				.AutoWrapText(true)
+				.Text_Lambda([this]() -> FText {
+					return ToolDropDownEntries[CurrentlySelectedTool]->Description;
+				})
+			]
 		];
 
 	ToolbarContainer->SetContent(ToolbarContent);
diff --git a/Engine/Source/Editor/MergeActors/Private/SMergeActorsToolbar.h b/Engine/Source/Editor/MergeActors/Private/SMergeActorsToolbar.h
--- a/Engine/Source/Editor/MergeActors/Private/SMergeActorsToolbar.h
+++ b/Engine/Source/Editor/MergeActors/Private/SMergeActorsToolbar.h
@@ -48,10 +48,11 @@ public:
 private:
 	struct FDropDownItem
 	{
-		FDropDownItem(const FText& InName, const FName& InIconName);
+		FDropDownItem(const FText& InName, const FName& InIconName, const FText& InDescription);
 
 		FText Name;
 		FName IconName;
+		FText Description;
 	};
 
 	/** Called to create the combo box entries */
diff --git a/Engine/Source/Editor/MergeActors/Public/IMergeActorsTool.h b/Engine/Source/Editor/MergeActors/Public/IMergeActorsTool.h
--- a/Engine/Source/Editor/MergeActors/Public/IMergeActorsTool.h
+++ b/Engine/Source/Editor/MergeActors/Public/IMergeActorsTool.h
@@ -25,6 +25,11 @@ public:
 	 */
 	virtual FName GetIconName() const = 0;
 
+	/**
+	 * Get tool name text to be displayed in the menus & Merge Actors toolbar
+	 */
+	virtual FText GetToolNameText() const = 0;
+
 	/**
 	 * Get Tooltip text displayed in the Merge Actors toolbar
 	 */
diff --git a/Engine/Source/Runtime/Engine/Classes/Engine/MeshMerging.h b/Engine/Source/Runtime/Engine/Classes/Engine/MeshMerging.h
--- a/Engine/Source/Runtime/Engine/Classes/Engine/MeshMerging.h
+++ b/Engine/Source/Runtime/Engine/Classes/Engine/MeshMerging.h
@@ -308,7 +308,7 @@ struct FMeshProxySettings
 	UPROPERTY(BlueprintReadWrite, EditAnywhere, Category = ProxySettings)
 	uint8 bReuseMeshLightmapUVs:1;
 
-	/** Bake identical meshes (or mesh instances) only once. Can lead to discrepancies with the source mesh visual, especially if WPO is used in the material(s). However, this will result in better quality baked textures & greatly reduce baking time. */
+	/** Bake identical meshes (or mesh instances) only once. Can lead to discrepancies with the source mesh visual, especially for materials that are using world position or per instance data. However, this will result in better quality baked textures & greatly reduce baking time. */
 	UPROPERTY(BlueprintReadWrite, EditAnywhere, Category = ProxySettings)
 	uint8 bGroupIdenticalMeshesForBaking:1;
 
