Crash when trying to run Merge Actor Tool using new Approximation and replace an already generated mesh

Will now generate the assets in temp packages & replace the original ones after.

diff --git a/Engine/Source/Editor/MergeActors/Private/MeshApproximationTool/MeshApproximationTool.cpp b/Engine/Source/Editor/MergeActors/Private/MeshApproximationTool/MeshApproximationTool.cpp
--- a/Engine/Source/Editor/MergeActors/Private/MeshApproximationTool/MeshApproximationTool.cpp
+++ b/Engine/Source/Editor/MergeActors/Private/MeshApproximationTool/MeshApproximationTool.cpp
@@ -28,7 +28,8 @@
 
 #include "AssetRegistryModule.h"
 #include "AssetToolsModule.h"
-
+#include "ObjectTools.h"
+#include "Algo/ForEach.h"
 
 #include "IGeometryProcessingInterfacesModule.h"
 #include "GeometryProcessingInterfaces/ApproximateActors.h"
@@ -225,28 +226,59 @@ bool FMeshApproximationTool::RunMerge(const FString& PackageName, const TArray<T
 	//
 
 	IGeometryProcessing_ApproximateActors::FOptions Options = ApproxActorsAPI->ConstructOptions(UseSettings);
-	Options.BasePackagePath = PackageName;
-
 	Options.TextureSizePolicy = TextureSizePolicy;
 	Options.MeshTexelDensity = TexelDensityPerMeter;
+
+	// Use temp packages - Needed to allow proper replacement of existing assets (performed below)
+	FString PackagePath = FPackageName::GetLongPackagePath(PackageName);
+	FString AssetName = FPackageName::GetLongPackageAssetName(PackageName);
+	Options.BasePackagePath = FString::Printf(TEXT("/Temp/Editor/MergeActors/MeshApproximationTool/%s"), *AssetName);
 	
-	// run actor approximation computation
+	// Run actor approximation computation
 	IGeometryProcessing_ApproximateActors::FResults Results;
 	ApproxActorsAPI->ApproximateActors(Actors, Options, Results);
 
-	// Notify Asset Browser about new Assets
-	for (UStaticMesh* StaticMesh : Results.NewMeshAssets)
+	auto ProcessNewAsset = [&PackagePath](UObject* NewAsset)
 	{
-		FAssetRegistryModule::AssetCreated(StaticMesh);
-	}
-	for (UMaterialInterface* Material : Results.NewMaterials)
-	{
-		FAssetRegistryModule::AssetCreated(Material);
-	}
-	for (UTexture2D* Texture : Results.NewTextures)
-	{
-		FAssetRegistryModule::AssetCreated(Texture);
-	}
+		// Move asset out of the temp package and into its final package
+		{
+			FString AssetName = NewAsset->GetName();
+			FString TargetPackageName = FPaths::Combine(PackagePath, AssetName);
+
+			UPackage* TargetPackage = CreatePackage(*TargetPackageName);
+			TargetPackage->FullyLoad();
+
+			// Remplace existing asset
+			UObject* AssetToReplace = StaticFindObjectFast(UObject::StaticClass(), TargetPackage, *AssetName);
+			if (AssetToReplace)
+			{
+				// Replace references
+				TArray<UObject*> ObjectsToReplace(&AssetToReplace, 1);
+				ObjectTools::ForceReplaceReferences(NewAsset, ObjectsToReplace);
+
+				// Move the previous asset to the transient package
+				AssetToReplace->Rename(nullptr, GetTransientPackage(), REN_DontCreateRedirectors | REN_NonTransactional | REN_ForceNoResetLoaders);
+			}
+
+			UPackage* TempPackage = NewAsset->GetPackage();
+
+			// Rename the asset to its final destination
+			NewAsset->Rename(nullptr, TargetPackage, REN_DontCreateRedirectors | REN_NonTransactional | REN_ForceNoResetLoaders);
+			check(NewAsset->HasAllFlags(RF_Public | RF_Standalone));
+
+			// Clean up flags on the temp package. It is not useful anymore.
+			TempPackage->ClearDirtyFlag();
+			TempPackage->SetFlags(RF_Transient);
+			TempPackage->ClearFlags(RF_Public | RF_Standalone);
+		}
+
+		// Notify Asset Browser about new Assets
+		FAssetRegistryModule::AssetCreated(NewAsset);
+	};
+
+	Algo::ForEach(Results.NewMeshAssets, ProcessNewAsset);
+	Algo::ForEach(Results.NewMaterials, ProcessNewAsset);
+	Algo::ForEach(Results.NewTextures, ProcessNewAsset);
 
 	if (ensure(Results.NewMeshAssets.Num() == 1))
 	{
