Fix crash when trying to run Merge Actor Tool using new Approximation and replace an already generated mesh.

diff --git a/Engine/Source/Editor/MergeActors/Private/MeshApproximationTool/MeshApproximationTool.cpp b/Engine/Source/Editor/MergeActors/Private/MeshApproximationTool/MeshApproximationTool.cpp
--- a/Engine/Source/Editor/MergeActors/Private/MeshApproximationTool/MeshApproximationTool.cpp
+++ b/Engine/Source/Editor/MergeActors/Private/MeshApproximationTool/MeshApproximationTool.cpp
@@ -230,19 +230,21 @@ bool FMeshApproximationTool::RunMerge(const FString& PackageName, const TArray<T
 	Options.MeshTexelDensity = TexelDensityPerMeter;
 
 	// Use temp packages - Needed to allow proper replacement of existing assets (performed below)
+	const FString NewAssetNamePrefix(TEXT("NEWASSET_"));
 	FString PackagePath = FPackageName::GetLongPackagePath(PackageName);
 	FString AssetName = FPackageName::GetLongPackageAssetName(PackageName);
-	Options.BasePackagePath = FString::Printf(TEXT("/Temp/Editor/MergeActors/MeshApproximationTool/%s"), *AssetName);
+	Options.BasePackagePath = PackagePath / NewAssetNamePrefix + AssetName;
 	
 	// Run actor approximation computation
 	IGeometryProcessing_ApproximateActors::FResults Results;
 	ApproxActorsAPI->ApproximateActors(Actors, Options, Results);
 
-	auto ProcessNewAsset = [&PackagePath](UObject* NewAsset)
+	auto ProcessNewAsset = [&PackagePath, &NewAssetNamePrefix](UObject* NewAsset)
 	{
 		// Move asset out of the temp package and into its final package
 		{
 			FString AssetName = NewAsset->GetName();
+			AssetName.RemoveFromStart(NewAssetNamePrefix);
 			FString TargetPackageName = FPaths::Combine(PackagePath, AssetName);
 
 			UPackage* TargetPackage = CreatePackage(*TargetPackageName);
@@ -263,7 +265,7 @@ bool FMeshApproximationTool::RunMerge(const FString& PackageName, const TArray<T
 			UPackage* TempPackage = NewAsset->GetPackage();
 
 			// Rename the asset to its final destination
-			NewAsset->Rename(nullptr, TargetPackage, REN_DontCreateRedirectors | REN_NonTransactional | REN_ForceNoResetLoaders);
+			NewAsset->Rename(*AssetName, TargetPackage, REN_DontCreateRedirectors | REN_NonTransactional | REN_ForceNoResetLoaders);
 			check(NewAsset->HasAllFlags(RF_Public | RF_Standalone));
 
 			// Clean up flags on the temp package. It is not useful anymore.
