Merge Actor - Mesh Approximate Tool - Improved error handling

diff --git a/Engine/Source/Editor/MergeActors/Private/MeshApproximationTool/MeshApproximationTool.cpp b/Engine/Source/Editor/MergeActors/Private/MeshApproximationTool/MeshApproximationTool.cpp
--- a/Engine/Source/Editor/MergeActors/Private/MeshApproximationTool/MeshApproximationTool.cpp
+++ b/Engine/Source/Editor/MergeActors/Private/MeshApproximationTool/MeshApproximationTool.cpp
@@ -278,6 +278,28 @@ bool FMeshApproximationTool::RunMerge(const FString& PackageName, const TArray<T
 		FAssetRegistryModule::AssetCreated(NewAsset);
 	};
 
+	if (Results.ResultCode != IGeometryProcessing_ApproximateActors::EResultCode::Success)
+	{
+		FText Message;
+		FText Title = LOCTEXT("ApproximateActorsFailed_Title", "Failed to merge actors");
+
+		switch (Results.ResultCode)
+		{
+		case IGeometryProcessing_ApproximateActors::EResultCode::MeshGenerationFailed:
+			Message = LOCTEXT("ApproximateActors_MeshGeneratedFailed", "Mesh generation failed. Please review the merge settings and look for additional errors in the console log.");
+			break;
+		case IGeometryProcessing_ApproximateActors::EResultCode::MaterialGenerationFailed:
+			Message = LOCTEXT("ApproximateActors_MaterialGenerationFailed", "Material generation failed. Please review the merge settings and look for additional errors in the console log.");
+			break;
+		case IGeometryProcessing_ApproximateActors::EResultCode::UnknownError:
+			Message = LOCTEXT("ApproximateActors_UnknownError", "Unknown merge error. Please review the merge settings and look for additional errors in the console log.");
+			break;
+		}		
+		
+		FMessageDialog::Open(EAppMsgType::Ok, Message, &Title);
+		return false;
+	}
+
 	Algo::ForEach(Results.NewMeshAssets, ProcessNewAsset);
 	Algo::ForEach(Results.NewMaterials, ProcessNewAsset);
 	Algo::ForEach(Results.NewTextures, ProcessNewAsset);
@@ -291,7 +313,7 @@ bool FMeshApproximationTool::RunMerge(const FString& PackageName, const TArray<T
 			ReplaceSourceActorsByApproximationMeshes(NewAssetsToSync, UniqueLevels[0], Actors);
 		}
 	}
-
+	
 	return true;
 }
 
