Merge Actor - Fixed typo causing MeshApproximation to exlude valid primitive components from the merge

diff --git a/Engine/Source/Editor/MergeActors/Private/MeshApproximationTool/MeshApproximationTool.cpp b/Engine/Source/Editor/MergeActors/Private/MeshApproximationTool/MeshApproximationTool.cpp
--- a/Engine/Source/Editor/MergeActors/Private/MeshApproximationTool/MeshApproximationTool.cpp
+++ b/Engine/Source/Editor/MergeActors/Private/MeshApproximationTool/MeshApproximationTool.cpp
@@ -182,7 +182,7 @@ bool FMeshApproximationTool::RunMerge(const FString& PackageName, const TArray<T
 	IGeometryProcessing_ApproximateActors::FOptions Options = ApproxActorsAPI->ConstructOptions(UseSettings);
 
 	TArray<UPrimitiveComponent*> PrimitiveComponents;
-	auto IsValidPrimitiveComponent = [](const TSharedPtr<FMergeComponentData>& Component) { return Component.IsValid() && !Component.Get()->bShouldIncorporate; };
+	auto IsValidPrimitiveComponent = [](const TSharedPtr<FMergeComponentData>& Component) { return Component.IsValid() && Component.Get()->bShouldIncorporate; };
 	auto GetPrimitiveComponent = [](const TSharedPtr<FMergeComponentData>& Component) { return Component.Get()->PrimComponent.Get(); };
 	Algo::TransformIf(SelectedComponents, PrimitiveComponents, IsValidPrimitiveComponent, GetPrimitiveComponent);
 
