Merge Actor - Textures created with the "Merge" option are not visible in the Content Browser

diff --git a/Engine/Source/Developer/MeshMergeUtilities/Private/MeshMergeUtilities.cpp b/Engine/Source/Developer/MeshMergeUtilities/Private/MeshMergeUtilities.cpp
--- a/Engine/Source/Developer/MeshMergeUtilities/Private/MeshMergeUtilities.cpp
+++ b/Engine/Source/Developer/MeshMergeUtilities/Private/MeshMergeUtilities.cpp
@@ -2473,12 +2473,12 @@ void FMeshMergeUtilities::MergeComponentsToStaticMesh(const TArray<UPrimitiveCom
 		if (InBasePackageName.IsEmpty())
 		{
 			AssetName = TEXT("SM_MERGED_") + FPackageName::GetShortName(MergedAssetPackageName);
-			PackageName = FPackageName::GetLongPackagePath(MergedAssetPackageName) + TEXT("/") + AssetName;
+			PackageName = FPackageName::GetLongPackagePath(MergedAssetPackageName) / AssetName;
 		}
 		else
 		{
-			AssetName = FPackageName::GetShortName(InBasePackageName);
-			PackageName = InBasePackageName;
+			AssetName = TEXT("SM_") + FPackageName::GetShortName(InBasePackageName);
+			PackageName = FPackageName::GetLongPackagePath(InBasePackageName) / AssetName;
 		}
 
 		UPackage* Package = InOuter;
diff --git a/Engine/Source/Editor/MergeActors/Private/MeshMergingTool/MeshMergingTool.cpp b/Engine/Source/Editor/MergeActors/Private/MeshMergingTool/MeshMergingTool.cpp
--- a/Engine/Source/Editor/MergeActors/Private/MeshMergingTool/MeshMergingTool.cpp
+++ b/Engine/Source/Editor/MergeActors/Private/MeshMergingTool/MeshMergingTool.cpp
@@ -21,6 +21,7 @@
 #include "ScopedTransaction.h"
 #include "MeshMergeModule.h"
 #include "ComponentReregisterContext.h"
+#include "ObjectTools.h"
 
 #define LOCTEXT_NAMESPACE "MeshMergingTool"
 
@@ -61,7 +62,7 @@ FText FMeshMergingTool::GetTooltipText() const
 
 FString FMeshMergingTool::GetDefaultPackageName() const
 {
-	FString PackageName = FPackageName::FilenameToLongPackageName(FPaths::ProjectContentDir() + TEXT("SM_MERGED"));
+	FString PackageName = FPackageName::FilenameToLongPackageName(FPaths::ProjectContentDir() + TEXT("MERGED"));
 
 	USelection* SelectedActors = GEditor->GetSelectedActors();
 	// Iterate through selected actors and find first static mesh asset
@@ -149,11 +150,41 @@ bool FMeshMergingTool::RunMerge(const FString& PackageName, const TArray<TShared
 	if (AssetsToSync.Num())
 	{
 		FAssetRegistryModule& AssetRegistry = FModuleManager::Get().LoadModuleChecked<FAssetRegistryModule>("AssetRegistry");
-		int32 AssetCount = AssetsToSync.Num();
-		for (int32 AssetIndex = 0; AssetIndex < AssetCount; AssetIndex++)
+
+		for (UObject* AssetToSync : AssetsToSync)
 		{
-			AssetRegistry.AssetCreated(AssetsToSync[AssetIndex]);
-			GEditor->BroadcastObjectReimported(AssetsToSync[AssetIndex]);
+			// MergeComponentsToStaticMesh() will have outered all assets (material instance, textures) to the static mesh package.
+			// Move each of them to their own package, so that they show up in the Content Browser
+			if (AssetToSync && !AssetToSync->IsA<UStaticMesh>())
+			{
+				FString AssetName = AssetToSync->GetName();
+				FString AssetPackagePath = FPackageName::GetLongPackagePath(AssetToSync->GetPathName());
+				FString AssetPackageName = AssetPackagePath / AssetName;
+
+				UPackage* AssetPackage = CreatePackage(*AssetPackageName);
+				check(AssetPackage);
+				AssetPackage->FullyLoad();
+				AssetPackage->Modify();
+
+				// Replace existing asset by the new one.
+				if (UObject* OldAsset = FindObject<UObject>(AssetPackage, *AssetName))
+				{
+					FName ObjectName = OldAsset->GetFName();
+					UObject* Outer = OldAsset->GetOuter();
+					OldAsset->Rename(nullptr, GetTransientPackage(), REN_DoNotDirty | REN_DontCreateRedirectors);
+												
+					// Consolidate or "Replace" the old object with the new object for any living references.
+					bool bShowDeleteConfirmation = false;
+					TArray<UObject*> OldDataAssetArray = { OldAsset };
+					ObjectTools::ConsolidateObjects(AssetToSync, { OldDataAssetArray }, bShowDeleteConfirmation);
+				}
+
+				AssetToSync->Rename(*AssetName, AssetPackage, REN_DontCreateRedirectors);
+				AssetToSync->SetFlags(RF_Public | RF_Standalone);
+			}
+
+			AssetRegistry.AssetCreated(AssetToSync);
+			GEditor->BroadcastObjectReimported(AssetToSync);
 		}
 
 		//Also notify the content browser that the new assets exists
