Merge Actor Widget.  Removing support for old thirdparty variant  of mesh merging at the request of the thirdparty.

diff --git a/Engine/Source/Editor/MergeActors/Private/MergeActorsModule.cpp b/Engine/Source/Editor/MergeActors/Private/MergeActorsModule.cpp
--- a/Engine/Source/Editor/MergeActors/Private/MergeActorsModule.cpp
+++ b/Engine/Source/Editor/MergeActors/Private/MergeActorsModule.cpp
@@ -118,17 +118,11 @@ void FMergeActorsModule::StartupModule()
 	IMeshReductionManagerModule& MeshReductionModule = FModuleManager::Get().LoadModuleChecked<IMeshReductionManagerModule>("MeshReductionInterface");
 	if (MeshReductionModule.GetMeshMergingInterface() != nullptr)
 	{
-		// Choose the correct UI.  This isn't ideal.
 		if (MeshReductionModule.GetMeshMergingInterface()->GetName().Equals(FString("ProxyLODMeshMerging")))
 		{
 			// Our Native tool
 			ensure(RegisterMergeActorsTool(MakeUnique<FMeshProxyTool>()));
 		}
-		else
-		{
-			// This is the Simplygon tool
-			ensure(RegisterMergeActorsTool(MakeUnique<FThirdPartyMeshProxyTool>()));
-		}
 	}
 
 	ensure(RegisterMergeActorsTool(MakeUnique<FMeshInstancingTool>()));
diff --git a/Engine/Source/Editor/MergeActors/Private/MeshProxyTool/MeshProxyTool.cpp b/Engine/Source/Editor/MergeActors/Private/MeshProxyTool/MeshProxyTool.cpp
--- a/Engine/Source/Editor/MergeActors/Private/MeshProxyTool/MeshProxyTool.cpp
+++ b/Engine/Source/Editor/MergeActors/Private/MeshProxyTool/MeshProxyTool.cpp
@@ -221,166 +221,6 @@ bool FMeshProxyTool::RunMerge(const FString& PackageName, const TArray<TSharedPt
 	return true;
 }
 
-FName FThirdPartyMeshProxyTool::GetIconName() const 
-{
-	return "MergeActors.MeshProxyTool";
-}
-
-FText FThirdPartyMeshProxyTool::GetToolNameText() const
-{
-	return LOCTEXT("ThirdPartyMeshProxyToolName", "Third-Party Simplification");
-}
-
-FText FThirdPartyMeshProxyTool::GetTooltipText() const
-{
-	return LOCTEXT("ThirdPartyMeshProxyToolTooltip", "Merge source actors meshes and perform a simplification pass. Will generate a single static mesh with baked textures.");
-}
-
-FString FThirdPartyMeshProxyTool::GetDefaultPackageName() const
-{
-	FString PackageName;
-
-	USelection* SelectedActors = GEditor->GetSelectedActors();
-
-	// Iterate through selected actors and find first static mesh asset
-	// Use this static mesh path as destination package name for a merged mesh
-	for (FSelectionIterator Iter(*SelectedActors); Iter; ++Iter)
-	{
-		AActor* Actor = Cast<AActor>(*Iter);
-		if (Actor)
-		{
-			TInlineComponentArray<UStaticMeshComponent*> SMComponets;
-			Actor->GetComponents(SMComponets);
-			for (UStaticMeshComponent* Component : SMComponets)
-			{
-				if (Component->GetStaticMesh())
-				{
-					PackageName = FPackageName::GetLongPackagePath(Component->GetStaticMesh()->GetOutermost()->GetName());
-					PackageName += FString(TEXT("/PROXY_")) + Component->GetStaticMesh()->GetName();
-					break;
-				}
-			}
-		}
-
-		if (!PackageName.IsEmpty())
-		{
-			break;
-		}
-	}
-
-	if (PackageName.IsEmpty())
-	{
-		PackageName = FPackageName::FilenameToLongPackageName(FPaths::ProjectContentDir() + TEXT("PROXY"));
-		PackageName = MakeUniqueObjectName(NULL, UPackage::StaticClass(), *PackageName).ToString();
-	}
-
-	return PackageName;
-}
-
-bool FThirdPartyMeshProxyTool::RunMergeFromSelection()
-{
-	FString PackageName;
-	if (GetPackageNameForMergeAction(GetDefaultPackageName(), PackageName))
-	{
-		return RunMerge(PackageName);
-	}
-	else
-	{
-		return false;
-	}
-}
-
-bool FThirdPartyMeshProxyTool::RunMergeFromWidget()
-{
-	FString PackageName;
-	if (GetPackageNameForMergeAction(GetDefaultPackageName(), PackageName))
-	{
-		return RunMerge(PackageName);
-	}
-	else
-	{
-		return false;
-	}
-}
-
-bool FThirdPartyMeshProxyTool::RunMerge(const FString& PackageName)
-{
-	TArray<AActor*> Actors;
-	TArray<UObject*> AssetsToSync;
-
-	// Get the module for the mesh merge utilities
-	const IMeshMergeUtilities& MeshMergeUtilities = FModuleManager::Get().LoadModuleChecked<IMeshMergeModule>("MeshMergeUtilities").GetUtilities();
-
-	USelection* SelectedActors = GEditor->GetSelectedActors();
-	for (FSelectionIterator Iter(*SelectedActors); Iter; ++Iter)
-	{
-		AActor* Actor = Cast<AActor>(*Iter);
-		if (Actor)
-		{
-			Actors.Add(Actor);
-		}
-	}
-
-	if (Actors.Num())
-	{
-		GWarn->BeginSlowTask(LOCTEXT("MeshProxy_CreatingProxy", "Creating Mesh Proxy"), true);
-		GEditor->BeginTransaction(LOCTEXT("MeshProxy_Create", "Creating Mesh Proxy"));
-
-		FVector ProxyLocation = FVector::ZeroVector;
-		TArray<UObject*> NewAssetsToSync;
-		
-		FCreateProxyDelegate ProxyDelegate;
-		ProxyDelegate.BindLambda(
-			[&NewAssetsToSync](const FGuid Guid, TArray<UObject*>& InAssetsToSync)
-		{
-			//Update the asset registry that a new static mash and material has been created
-			if (InAssetsToSync.Num())
-			{
-				FAssetRegistryModule& AssetRegistry = FModuleManager::Get().LoadModuleChecked<FAssetRegistryModule>("AssetRegistry");
-				int32 AssetCount = InAssetsToSync.Num();
-				for (int32 AssetIndex = 0; AssetIndex < AssetCount; AssetIndex++)
-				{
-					AssetRegistry.AssetCreated(InAssetsToSync[AssetIndex]);
-					GEditor->BroadcastObjectReimported(InAssetsToSync[AssetIndex]);
-				}
-
-				//Also notify the content browser that the new assets exists
-				FContentBrowserModule& ContentBrowserModule = FModuleManager::Get().LoadModuleChecked<FContentBrowserModule>("ContentBrowser");
-				ContentBrowserModule.Get().SyncBrowserToAssets(InAssetsToSync, true);
-
-				NewAssetsToSync += InAssetsToSync;
-			}
-		});		
-
-		FGuid JobGuid = FGuid::NewGuid();
-		MeshMergeUtilities.CreateProxyMesh(Actors, ProxySettings, nullptr, PackageName, JobGuid, ProxyDelegate);
-
-		if (bReplaceSourceActors && Actors[0]->GetLevel())
-		{
-			ReplaceSourceActorsByProxyMesh(NewAssetsToSync, Actors[0]->GetLevel(), Actors);
-		}
-
-		GEditor->EndTransaction();
-		GWarn->EndSlowTask();
-	}
-
-	return true;
-}
-
-bool FThirdPartyMeshProxyTool::CanMergeFromSelection() const
-{
-	return true;
-}
-
-bool FThirdPartyMeshProxyTool::CanMergeFromWidget() const
-{
-	return true;
-}
-
-TSharedRef<SWidget> FThirdPartyMeshProxyTool::GetWidget()
-{
-	return SNew(SThirdPartyMeshProxyDialog, this);
-}
 
 #undef LOCTEXT_NAMESPACE
 
diff --git a/Engine/Source/Editor/MergeActors/Private/MeshProxyTool/MeshProxyTool.h b/Engine/Source/Editor/MergeActors/Private/MeshProxyTool/MeshProxyTool.h
--- a/Engine/Source/Editor/MergeActors/Private/MeshProxyTool/MeshProxyTool.h
+++ b/Engine/Source/Editor/MergeActors/Private/MeshProxyTool/MeshProxyTool.h
@@ -87,32 +87,3 @@ protected:
 	UMeshProxySettingsObject* SettingsObject;
 };
 
-/**
-* Third Party Mesh Proxy Tool
-*/
-class FThirdPartyMeshProxyTool : public IMergeActorsTool
-{
-	friend class SThirdPartyMeshProxyDialog;
-
-public:
-
-	// IMergeActorsTool interface
-	virtual TSharedRef<SWidget> GetWidget() override;
-	virtual FName GetIconName() const override;
-	virtual FText GetToolNameText() const override;
-	virtual FText GetTooltipText() const override;
-	virtual FString GetDefaultPackageName() const override;
-	virtual bool GetReplaceSourceActors() const override { return bReplaceSourceActors; }
-	virtual void SetReplaceSourceActors(bool bInReplaceSourceActors) override { bReplaceSourceActors = bInReplaceSourceActors; }
-	virtual bool RunMergeFromSelection() override;
-	virtual bool RunMergeFromWidget() override;
-	virtual bool CanMergeFromSelection() const override;
-	virtual bool CanMergeFromWidget() const override;
-
-protected:
-	virtual bool RunMerge(const FString& PackageName);
-
-protected:
-	bool bReplaceSourceActors = false;
-	FMeshProxySettings ProxySettings;
-};
diff --git a/Engine/Source/Editor/MergeActors/Private/MeshProxyTool/SMeshProxyDialog.cpp b/Engine/Source/Editor/MergeActors/Private/MeshProxyTool/SMeshProxyDialog.cpp
--- a/Engine/Source/Editor/MergeActors/Private/MeshProxyTool/SMeshProxyDialog.cpp
+++ b/Engine/Source/Editor/MergeActors/Private/MeshProxyTool/SMeshProxyDialog.cpp
@@ -51,446 +51,5 @@ void  SMeshProxyDialog::Construct(const FArguments& InArgs, FMeshProxyTool* InTo
 
 END_SLATE_FUNCTION_BUILD_OPTIMIZATION
 
-void SThirdPartyMeshProxyDialog::Construct(const FArguments& InArgs, FThirdPartyMeshProxyTool* InTool)
-{
-	Tool = InTool;
-	check(Tool != nullptr);
-
-	CuttingPlaneOptions.Add(MakeShareable(new FString(TEXT("+X"))));
-	CuttingPlaneOptions.Add(MakeShareable(new FString(TEXT("+Y"))));
-	CuttingPlaneOptions.Add(MakeShareable(new FString(TEXT("+Z"))));
-	CuttingPlaneOptions.Add(MakeShareable(new FString(TEXT("-X"))));
-	CuttingPlaneOptions.Add(MakeShareable(new FString(TEXT("-Y"))));
-	CuttingPlaneOptions.Add(MakeShareable(new FString(TEXT("-Z"))));
-
-	TextureResolutionOptions.Add(MakeShareable(new FString(TEXT("64"))));
-	TextureResolutionOptions.Add(MakeShareable(new FString(TEXT("128"))));
-	TextureResolutionOptions.Add(MakeShareable(new FString(TEXT("256"))));
-	TextureResolutionOptions.Add(MakeShareable(new FString(TEXT("512"))));
-	TextureResolutionOptions.Add(MakeShareable(new FString(TEXT("1024"))));
-	TextureResolutionOptions.Add(MakeShareable(new FString(TEXT("2048"))));
-
-	CreateLayout();
-}
-
-BEGIN_SLATE_FUNCTION_BUILD_OPTIMIZATION
-void  SThirdPartyMeshProxyDialog::CreateLayout()
-{
-	int32 TextureResEntryIndex = FindTextureResolutionEntryIndex(Tool->ProxySettings.MaterialSettings.TextureSize.X);
-	int32 LightMapResEntryIndex = FindTextureResolutionEntryIndex(Tool->ProxySettings.LightMapResolution);
-	TextureResEntryIndex = FMath::Max(TextureResEntryIndex, 0);
-	LightMapResEntryIndex = FMath::Max(LightMapResEntryIndex, 0);
-		
-	this->ChildSlot
-	[
-		SNew(SVerticalBox)
-			
-		+SVerticalBox::Slot()
-		.AutoHeight()
-		.HAlign(HAlign_Center)
-		.VAlign(VAlign_Center)
-		.Padding(10)
-		[
-			// Simplygon logo
-			SNew(SImage)
-			.Image(FAppStyle::GetBrush("MeshProxy.SimplygonLogo"))
-		]
-			
-		+SVerticalBox::Slot()
-		.AutoHeight()
-		[
-			SNew(SBorder)
-			.BorderImage(FAppStyle::GetBrush("ToolPanel.GroupBorder"))
-			[
-				// Proxy options
-				SNew(SVerticalBox)
-
-				+SVerticalBox::Slot()
-				.AutoHeight()
-				.Padding(FAppStyle::GetMargin("StandardDialog.ContentPadding"))
-				[
-					SNew(SHorizontalBox)
-					+SHorizontalBox::Slot()
-					.FillWidth(0.5f)
-					.VAlign(VAlign_Center)
-					[
-						SNew(STextBlock)
-						.Text(LOCTEXT("OnScreenSizeLabel", "On Screen Size (pixels)"))
-						.Font(FAppStyle::GetFontStyle("StandardDialog.SmallFont"))
-						.ToolTipText(GetPropertyToolTipText(GET_MEMBER_NAME_CHECKED(FMeshProxySettings, ScreenSize)))
-					]
-					+ SHorizontalBox::Slot()
-					.FillWidth(0.5f)
-					.HAlign(HAlign_Left)
-					.VAlign(VAlign_Center)
-					[
-						SNew(SBox)
-						.HAlign(HAlign_Fill)
-						.MinDesiredWidth(100.0f)
-						.MaxDesiredWidth(100.0f)
-						[
-							SNew(SNumericEntryBox<int32>)
-							.Font(FAppStyle::GetFontStyle("StandardDialog.SmallFont"))
-							.MinValue(40)
-							.MaxValue(1200)
-							.MinSliderValue(40)
-							.MaxSliderValue(1200)
-							.AllowSpin(true)
-							.Value(this, &SThirdPartyMeshProxyDialog::GetScreenSize)
-							.OnValueChanged(this, &SThirdPartyMeshProxyDialog::ScreenSizeChanged)
-						]
-					]
-				]
-
-				+SVerticalBox::Slot()
-				.AutoHeight()
-				.Padding(FAppStyle::GetMargin("StandardDialog.ContentPadding"))
-				[
-					SNew(SHorizontalBox)
-					+SHorizontalBox::Slot()
-					.FillWidth(0.5f)
-					.VAlign(VAlign_Center)
-					[
-						SNew(STextBlock)
-						.Text(LOCTEXT("MergeDistanceLabel", "Merge Distance (pixels)"))
-						.Font(FAppStyle::GetFontStyle("StandardDialog.SmallFont"))
-						.ToolTipText(GetPropertyToolTipText(GET_MEMBER_NAME_CHECKED(FMeshProxySettings, MergeDistance)))
-					]
-					+ SHorizontalBox::Slot()
-					.FillWidth(0.5f)
-					.HAlign(HAlign_Left)
-					.VAlign(VAlign_Center)
-					[
-						SNew(SBox)
-						.HAlign(HAlign_Fill)
-						.MinDesiredWidth(100.0f)
-						.MaxDesiredWidth(100.0f)
-						[
-							SNew(SNumericEntryBox<int32>)
-							.Font(FAppStyle::GetFontStyle("StandardDialog.SmallFont"))
-							.MinValue(0)
-							.MaxValue(300)
-							.MinSliderValue(0)
-							.MaxSliderValue(300)
-							.AllowSpin(true)
-							.Value(this, &SThirdPartyMeshProxyDialog::GetMergeDistance)
-							.OnValueChanged(this, &SThirdPartyMeshProxyDialog::MergeDistanceChanged)
-						]
-					]
-				]
-
-				+SVerticalBox::Slot()
-				.AutoHeight()
-				.Padding(FAppStyle::GetMargin("StandardDialog.ContentPadding"))
-				[
-					SNew(SHorizontalBox)
-					+SHorizontalBox::Slot()
-					.FillWidth(0.5f)
-					.VAlign(VAlign_Center)
-					[
-						SNew(STextBlock)
-						.Text(LOCTEXT("TextureResolutionLabel", "Texture Resolution"))
-						.Font(FAppStyle::GetFontStyle("StandardDialog.SmallFont"))
-					]
-					+ SHorizontalBox::Slot()
-					.FillWidth(0.5f)
-					.HAlign(HAlign_Left)
-					.VAlign(VAlign_Center)
-					[
-						SNew(STextComboBox)
-						.Font(FAppStyle::GetFontStyle("StandardDialog.SmallFont"))
-						.OptionsSource(&TextureResolutionOptions)
-						.InitiallySelectedItem(TextureResolutionOptions[TextureResEntryIndex])
-						.OnSelectionChanged(this, &SThirdPartyMeshProxyDialog::SetTextureResolution)
-					]
-				]
-
-				+SVerticalBox::Slot()
-				.AutoHeight()
-				.Padding(FAppStyle::GetMargin("StandardDialog.ContentPadding"))
-				[
-					SNew(SHorizontalBox)
-					+SHorizontalBox::Slot()
-					.FillWidth(0.5f)
-					.VAlign(VAlign_Center)
-					[
-						SNew(STextBlock)
-						.Text(LOCTEXT("LightMapResolutionLabel", "LightMap Resolution"))
-						.Font(FAppStyle::GetFontStyle("StandardDialog.SmallFont"))
-						.ToolTipText(GetPropertyToolTipText(GET_MEMBER_NAME_CHECKED(FMeshProxySettings, LightMapResolution)))
-					]
-					+ SHorizontalBox::Slot()
-					.FillWidth(0.5f)
-					.HAlign(HAlign_Left)
-					.VAlign(VAlign_Center)
-					[
-						SNew(STextComboBox)
-						.Font(FAppStyle::GetFontStyle("StandardDialog.SmallFont"))
-						.OptionsSource(&TextureResolutionOptions)
-						.InitiallySelectedItem(TextureResolutionOptions[LightMapResEntryIndex])
-						.OnSelectionChanged(this, &SThirdPartyMeshProxyDialog::SetLightMapResolution)
-					]
-				]
-
-				+SVerticalBox::Slot()
-				.AutoHeight()
-				.Padding(FAppStyle::GetMargin("StandardDialog.ContentPadding"))
-				[
-					SNew(SHorizontalBox)
-					+SHorizontalBox::Slot()
-					.FillWidth(0.5f)
-					.VAlign(VAlign_Center)
-					.Padding(0.0, 0.0, 3.0, 0.0)
-					[
-						SNew(STextBlock)
-						.Text(LOCTEXT("HardAngleLabel", "Hard Edge Angle"))
-						.Font(FAppStyle::GetFontStyle("StandardDialog.SmallFont"))
-						.ToolTipText(GetPropertyToolTipText(GET_MEMBER_NAME_CHECKED(FMeshProxySettings, HardAngleThreshold)))
-					]
-					+SHorizontalBox::Slot()
-					.FillWidth(0.5f)
-					.HAlign(HAlign_Left)
-					.VAlign(VAlign_Center)
-					[
-						SNew(SBox)
-						.HAlign(HAlign_Fill)
-						.MinDesiredWidth(100.0f)
-						.MaxDesiredWidth(100.0f)
-						[
-							SNew(SNumericEntryBox<float>)
-							.Font(FAppStyle::GetFontStyle("StandardDialog.SmallFont"))
-							.MinValue(0.f)
-							.MaxValue(180.f)
-							.MinSliderValue(0.f)
-							.MaxSliderValue(180.f)
-							.AllowSpin(true)
-							.Value(this, &SThirdPartyMeshProxyDialog::GetHardAngleThreshold)
-							.OnValueChanged(this, &SThirdPartyMeshProxyDialog::HardAngleThresholdChanged)
-							.IsEnabled(this, &SThirdPartyMeshProxyDialog::HardAngleThresholdEnabled)
-						]
-					]
-				]
-
-				+SVerticalBox::Slot()
-				.AutoHeight()
-				.Padding(FAppStyle::GetMargin("StandardDialog.ContentPadding"))
-				[
-					SNew(SCheckBox)
-					.Type(ESlateCheckBoxType::CheckBox)
-					.IsChecked(this, &SThirdPartyMeshProxyDialog::GetRecalculateNormals)
-					.OnCheckStateChanged(this, &SThirdPartyMeshProxyDialog::SetRecalculateNormals)
-					.Content()
-					[
-						SNew(STextBlock)
-						.Text(LOCTEXT("RecalcNormalsLabel", "Recalculate Normals"))
-						.Font(FAppStyle::GetFontStyle("StandardDialog.SmallFont"))
-						.ToolTipText(GetPropertyToolTipText(GET_MEMBER_NAME_CHECKED(FMeshProxySettings, bRecalculateNormals)))
-					]
-				]
-
-				+SVerticalBox::Slot()
-				.AutoHeight()
-				.Padding(FAppStyle::GetMargin("StandardDialog.ContentPadding"))
-				[
-					SNew(SCheckBox)
-					.Type(ESlateCheckBoxType::CheckBox)
-					.IsChecked(this, &SThirdPartyMeshProxyDialog::GetExportNormalMap)
-					.OnCheckStateChanged(this, &SThirdPartyMeshProxyDialog::SetExportNormalMap)
-					.Content()
-					[
-						SNew(STextBlock)
-						.Text(LOCTEXT("ExportNormalMapLabel", "Export Normal Map"))
-						.Font(FAppStyle::GetFontStyle("StandardDialog.SmallFont"))
-					]
-				]
-
-				+SVerticalBox::Slot()
-				.AutoHeight()
-				.Padding(FAppStyle::GetMargin("StandardDialog.ContentPadding"))
-				[
-					SNew(SCheckBox)
-					.Type(ESlateCheckBoxType::CheckBox)
-					.IsChecked(this, &SThirdPartyMeshProxyDialog::GetExportMetallicMap)
-					.OnCheckStateChanged(this, &SThirdPartyMeshProxyDialog::SetExportMetallicMap)
-					.Content()
-					[
-						SNew(STextBlock)
-						.Text(LOCTEXT("ExportMetallicMapLabel", "Export Metallic Map"))
-						.Font(FAppStyle::GetFontStyle("StandardDialog.SmallFont"))
-					]
-				]
-
-				+SVerticalBox::Slot()
-				.AutoHeight()
-				.Padding(FAppStyle::GetMargin("StandardDialog.ContentPadding"))
-				[
-					SNew(SCheckBox)
-					.Type(ESlateCheckBoxType::CheckBox)
-					.IsChecked(this, &SThirdPartyMeshProxyDialog::GetExportRoughnessMap)
-					.OnCheckStateChanged(this, &SThirdPartyMeshProxyDialog::SetExportRoughnessMap)
-					.Content()
-					[
-						SNew(STextBlock)
-						.Text(LOCTEXT("ExportRoughnessMapLabel", "Export Roughness Map"))
-						.Font(FAppStyle::GetFontStyle("StandardDialog.SmallFont"))
-					]
-				]
-
-				+SVerticalBox::Slot()
-				.AutoHeight()
-				.Padding(FAppStyle::GetMargin("StandardDialog.ContentPadding"))
-				[
-					SNew(SCheckBox)
-					.Type(ESlateCheckBoxType::CheckBox)
-					.IsChecked(this, &SThirdPartyMeshProxyDialog::GetExportSpecularMap)
-					.OnCheckStateChanged(this, &SThirdPartyMeshProxyDialog::SetExportSpecularMap)
-					.Content()
-					[
-						SNew(STextBlock)
-						.Text(LOCTEXT("ExportSpecularMapLabel", "Export Specular Map"))
-						.Font(FAppStyle::GetFontStyle("StandardDialog.SmallFont"))
-					]
-				]
-			]
-		]
-	];
-}
-END_SLATE_FUNCTION_BUILD_OPTIMIZATION
-
-int32 SThirdPartyMeshProxyDialog::FindTextureResolutionEntryIndex(int32 InResolution) const
-{
-	FString ResolutionStr = TTypeToString<int32>::ToString(InResolution);
-	
-	int32 Result = TextureResolutionOptions.IndexOfByPredicate([&](const TSharedPtr<FString>& Entry)
-	{
-		return (ResolutionStr == *Entry);
-	});
-
-	return Result;
-}
-
-FText SThirdPartyMeshProxyDialog::GetPropertyToolTipText(const FName& PropertyName) const
-{
-	FProperty* Property = FMeshProxySettings::StaticStruct()->FindPropertyByName(PropertyName);
-	if (Property)
-	{
-		return Property->GetToolTipText();
-	}
-	
-	return FText::GetEmpty();
-}
-
-//Screen size
-TOptional<int32> SThirdPartyMeshProxyDialog::GetScreenSize() const
-{
-	return Tool->ProxySettings.ScreenSize;
-}
-
-void SThirdPartyMeshProxyDialog::ScreenSizeChanged(int32 NewValue)
-{
-	Tool->ProxySettings.ScreenSize = NewValue;
-}
-
-//Recalculate normals
-ECheckBoxState SThirdPartyMeshProxyDialog::GetRecalculateNormals() const
-{
-	return Tool->ProxySettings.bRecalculateNormals ? ECheckBoxState::Checked: ECheckBoxState::Unchecked;
-}
-
-void SThirdPartyMeshProxyDialog::SetRecalculateNormals(ECheckBoxState NewValue)
-{
-	Tool->ProxySettings.bRecalculateNormals = (NewValue == ECheckBoxState::Checked);
-}
-
-//Hard Angle Threshold
-bool SThirdPartyMeshProxyDialog::HardAngleThresholdEnabled() const
-{
-	if(Tool->ProxySettings.bRecalculateNormals)
-	{
-		return true;
-	}
-
-	return false;
-}
-
-TOptional<float> SThirdPartyMeshProxyDialog::GetHardAngleThreshold() const
-{
-	return Tool->ProxySettings.HardAngleThreshold;
-}
-
-void SThirdPartyMeshProxyDialog::HardAngleThresholdChanged(float NewValue)
-{
-	Tool->ProxySettings.HardAngleThreshold = NewValue;
-}
-
-//Merge Distance
-TOptional<int32> SThirdPartyMeshProxyDialog::GetMergeDistance() const
-{
-	return UE::LWC::FloatToIntCastChecked<int32>(Tool->ProxySettings.MergeDistance);
-}
-
-void SThirdPartyMeshProxyDialog::MergeDistanceChanged(int32 NewValue)
-{
-	Tool->ProxySettings.MergeDistance = (float)NewValue;
-}
-
-//Texture Resolution
-void SThirdPartyMeshProxyDialog::SetTextureResolution(TSharedPtr<FString> NewSelection, ESelectInfo::Type SelectInfo)
-{
-	int32 Resolution = 512;
-	TTypeFromString<int32>::FromString(Resolution, **NewSelection);
-	FIntPoint TextureSize(Resolution, Resolution);
-	
-	Tool->ProxySettings.MaterialSettings.TextureSize = TextureSize;
-}
-
-void SThirdPartyMeshProxyDialog::SetLightMapResolution(TSharedPtr<FString> NewSelection, ESelectInfo::Type SelectInfo)
-{
-	int32 Resolution = 256;
-	TTypeFromString<int32>::FromString(Resolution, **NewSelection);
-		
-	Tool->ProxySettings.LightMapResolution = Resolution;
-}
-
-ECheckBoxState SThirdPartyMeshProxyDialog::GetExportNormalMap() const
-{
-	return Tool->ProxySettings.MaterialSettings.bNormalMap ? ECheckBoxState::Checked :  ECheckBoxState::Unchecked;
-}
-
-void SThirdPartyMeshProxyDialog::SetExportNormalMap(ECheckBoxState NewValue)
-{
-	Tool->ProxySettings.MaterialSettings.bNormalMap = (NewValue == ECheckBoxState::Checked);
-}
-
-ECheckBoxState SThirdPartyMeshProxyDialog::GetExportMetallicMap() const
-{
-	return Tool->ProxySettings.MaterialSettings.bMetallicMap ? ECheckBoxState::Checked :  ECheckBoxState::Unchecked;
-}
-
-void SThirdPartyMeshProxyDialog::SetExportMetallicMap(ECheckBoxState NewValue)
-{
-	Tool->ProxySettings.MaterialSettings.bMetallicMap = (NewValue == ECheckBoxState::Checked);
-}
-
-ECheckBoxState SThirdPartyMeshProxyDialog::GetExportRoughnessMap() const
-{
-	return Tool->ProxySettings.MaterialSettings.bRoughnessMap ? ECheckBoxState::Checked :  ECheckBoxState::Unchecked;
-}
-
-void SThirdPartyMeshProxyDialog::SetExportRoughnessMap(ECheckBoxState NewValue)
-{
-	Tool->ProxySettings.MaterialSettings.bRoughnessMap = (NewValue == ECheckBoxState::Checked);
-}
-
-ECheckBoxState SThirdPartyMeshProxyDialog::GetExportSpecularMap() const
-{
-	return Tool->ProxySettings.MaterialSettings.bSpecularMap ? ECheckBoxState::Checked :  ECheckBoxState::Unchecked;
-}
-
-void SThirdPartyMeshProxyDialog::SetExportSpecularMap(ECheckBoxState NewValue)
-{
-	Tool->ProxySettings.MaterialSettings.bSpecularMap = (NewValue == ECheckBoxState::Checked);
-}
-
 
 #undef LOCTEXT_NAMESPACE
diff --git a/Engine/Source/Editor/MergeActors/Private/MeshProxyTool/SMeshProxyDialog.h b/Engine/Source/Editor/MergeActors/Private/MeshProxyTool/SMeshProxyDialog.h
--- a/Engine/Source/Editor/MergeActors/Private/MeshProxyTool/SMeshProxyDialog.h
+++ b/Engine/Source/Editor/MergeActors/Private/MeshProxyTool/SMeshProxyDialog.h
@@ -42,67 +42,3 @@ private:
 	UMeshProxySettingsObject* ProxySettings;
 };
 
-
-
-class FThirdPartyMeshProxyTool;
-/*-----------------------------------------------------------------------------
-	SThirdPartyMeshProxyDialog  -- Used for Simplygon (third party tool) integration.
------------------------------------------------------------------------------*/
-class SThirdPartyMeshProxyDialog : public SCompoundWidget
-{
-public:
-	SLATE_BEGIN_ARGS(SThirdPartyMeshProxyDialog)
-	{
-	}
-	SLATE_END_ARGS()
-
-public:
-	/** SWidget functions */
-	void Construct(const FArguments& InArgs, FThirdPartyMeshProxyTool* InTool);
-
-protected:
-	/** ScreenSize accessors */
-	TOptional<int32> GetScreenSize() const;
-	void ScreenSizeChanged(int32 NewValue);		//used with editable text block (Simplygon)
-
-	/** Recalculate Normals accessors */
-	ECheckBoxState GetRecalculateNormals() const;
-	void SetRecalculateNormals(ECheckBoxState NewValue);
-
-	/** Hard Angle Threshold accessors */
-	TOptional<float> GetHardAngleThreshold() const;
-	bool HardAngleThresholdEnabled() const;
-	void HardAngleThresholdChanged(float NewValue);
-
-	/** Hole filling accessors */
-	TOptional<int32> GetMergeDistance() const;
-	void MergeDistanceChanged(int32 NewValue);
-
-	/** TextureResolution accessors */
-	void SetTextureResolution(TSharedPtr<FString> NewSelection, ESelectInfo::Type SelectInfo);
-	void SetLightMapResolution(TSharedPtr<FString> NewSelection, ESelectInfo::Type SelectInfo);
-
-	/** Export material properties acessors **/
-	ECheckBoxState GetExportNormalMap() const;
-	void SetExportNormalMap(ECheckBoxState NewValue);
-	ECheckBoxState GetExportMetallicMap() const;
-	void SetExportMetallicMap(ECheckBoxState NewValue);
-	ECheckBoxState GetExportRoughnessMap() const;
-	void SetExportRoughnessMap(ECheckBoxState NewValue);
-	ECheckBoxState GetExportSpecularMap() const;
-	void SetExportSpecularMap(ECheckBoxState NewValue);
-
-private:
-	/** Creates the geometry mode controls */
-	void CreateLayout();
-
-	int32 FindTextureResolutionEntryIndex(int32 InResolution) const;
-	FText GetPropertyToolTipText(const FName& PropertyName) const;
-
-private:
-	FThirdPartyMeshProxyTool* Tool;
-
-	TArray< TSharedPtr<FString> >	CuttingPlaneOptions;
-	TArray< TSharedPtr<FString> >	TextureResolutionOptions;
-};
-
