MeshMerge - Fix GC issue in Merge Actors tool

* ConsolidateObjects() can trigger GC, which could destroy merged assets held as raw UObject pointers. Root them with TStrongObjectPtr to keep them alive.
* Use the batch FReplaceRequest overload to perform a single reference-replacement pass instead of one per asset.

#Synced-CL 50652432

diff --git a/Engine/Source/Editor/MergeActors/Private/MeshMergingTool/MeshMergingTool.cpp b/Engine/Source/Editor/MergeActors/Private/MeshMergingTool/MeshMergingTool.cpp
--- a/Engine/Source/Editor/MergeActors/Private/MeshMergingTool/MeshMergingTool.cpp
+++ b/Engine/Source/Editor/MergeActors/Private/MeshMergingTool/MeshMergingTool.cpp
@@ -111,7 +111,12 @@ bool FMeshMergingTool::RunMerge(const FString& PackageName, const TArray<TShared
 	}
 
 	FVector MergedActorLocation;
-	TArray<UObject*> AssetsToSync;
+	TArray<UObject*> MergedMeshAssets;
+
+	// The call to ObjectTools::ConsolidateObjects() can actually trigger a GC, so we keep the newly generated assets
+	// in the root set using this array.
+	TArray<TStrongObjectPtr<UObject>> RootedObjects;
+
 	// Merge...
 	{
 		FScopedSlowTask SlowTask(0, LOCTEXT("MergingActorsSlowTask", "Merging actors..."));
@@ -140,64 +145,98 @@ bool FMeshMergingTool::RunMerge(const FString& PackageName, const TArray<TShared
 			if (FindObject<UObject>(nullptr, *PackageName))
 			{
 				FGlobalComponentReregisterContext GlobalReregister;
-				MeshUtilities.MergeComponentsToStaticMesh(ComponentsToMerge, World, SettingsObject->Settings, nullptr, nullptr, PackageName, AssetsToSync, MergedActorLocation, ScreenAreaSize, true);
+				MeshUtilities.MergeComponentsToStaticMesh(ComponentsToMerge, World, SettingsObject->Settings, nullptr, nullptr, PackageName, MergedMeshAssets, MergedActorLocation, ScreenAreaSize, true);
 			}
 			else
 			{
-				MeshUtilities.MergeComponentsToStaticMesh(ComponentsToMerge, World, SettingsObject->Settings, nullptr, nullptr, PackageName, AssetsToSync, MergedActorLocation, ScreenAreaSize, true);
+				MeshUtilities.MergeComponentsToStaticMesh(ComponentsToMerge, World, SettingsObject->Settings, nullptr, nullptr, PackageName, MergedMeshAssets, MergedActorLocation, ScreenAreaSize, true);
 			}
 		}
 	}
 
-	if (AssetsToSync.Num())
+	Algo::Transform(MergedMeshAssets, RootedObjects, [](UObject* AssetToSync) { return TStrongObjectPtr(AssetToSync); });
+
+	if (MergedMeshAssets.Num())
 	{
 		FAssetRegistryModule& AssetRegistry = FModuleManager::Get().LoadModuleChecked<FAssetRegistryModule>("AssetRegistry");
 
-		for (UObject* AssetToSync : AssetsToSync)
+		// MergeComponentsToStaticMesh() will have outered all assets (material instances, textures) to the static mesh package.
+		// Move each of them to their own package so that they show up in the Content Browser.
+		struct FAssetMoveInfo
 		{
-			// MergeComponentsToStaticMesh() will have outered all assets (material instance, textures) to the static mesh package.
-			// Move each of them to their own package, so that they show up in the Content Browser
-			if (AssetToSync && !AssetToSync->IsA<UStaticMesh>())
+			UObject* OldAsset;
+			UObject* NewAsset;
+			FString AssetName;
+			UPackage* AssetPackage;
+		};
+		TArray<FAssetMoveInfo> AssetMoveInfos;
+
+		// Create target packages
+		for (UObject* NewAsset : MergedMeshAssets)
+		{
+			if (NewAsset && !NewAsset->IsA<UStaticMesh>())
 			{
-				FString AssetName = AssetToSync->GetName();
-				FString AssetPackagePath = FPackageName::GetLongPackagePath(AssetToSync->GetPathName());
-				FString AssetPackageName = AssetPackagePath / AssetName;
-
-				UPackage* AssetPackage = CreatePackage(*AssetPackageName);
-				check(AssetPackage);
-				AssetPackage->FullyLoad();
-				AssetPackage->Modify();
-
-				// Replace existing asset by the new one.
-				if (UObject* OldAsset = FindObject<UObject>(AssetPackage, *AssetName))
+				FString NewAssetName = NewAsset->GetName();
+				FString NewAssetPackagePath = FPackageName::GetLongPackagePath(NewAsset->GetPathName());
+				FString NewAssetPackageName = NewAssetPackagePath / NewAssetName;
+
+				UPackage* NewAssetPackage = CreatePackage(*NewAssetPackageName);
+				check(NewAssetPackage);
+				NewAssetPackage->FullyLoad();
+				NewAssetPackage->Modify();
+				RootedObjects.Emplace(NewAssetPackage);
+
+				UObject* OldAsset = FindObject<UObject>(NewAssetPackage, *NewAssetName);
+				if (OldAsset)
 				{
-					FName ObjectName = OldAsset->GetFName();
-					UObject* Outer = OldAsset->GetOuter();
+					// Move old assets out of the way
 					OldAsset->Rename(nullptr, GetTransientPackage(), REN_DoNotDirty | REN_DontCreateRedirectors);
-												
-					// Consolidate or "Replace" the old object with the new object for any living references.
-					bool bShowDeleteConfirmation = false;
-					TArray<UObject*> OldDataAssetArray = { OldAsset };
-					ObjectTools::ConsolidateObjects(AssetToSync, { OldDataAssetArray }, bShowDeleteConfirmation);
 				}
+				AssetMoveInfos.Add({ OldAsset, NewAsset, MoveTemp(NewAssetName), NewAssetPackage });
+			}
+		}
+
+		// Replace all old asset references with the new ones
+		{
+			TArray<ObjectTools::FReplaceRequest> ConsolidateRequests;
+			for (FAssetMoveInfo& MoveInfo : AssetMoveInfos)
+			{
+				if (MoveInfo.OldAsset)
+				{
+					ConsolidateRequests.Add({MoveInfo.NewAsset, MakeArrayView(&MoveInfo.OldAsset, 1)});
+				}
+			}
 
-				AssetToSync->Rename(*AssetName, AssetPackage, REN_DontCreateRedirectors | REN_AllowPackageLinkerMismatch);
-				AssetToSync->SetFlags(RF_Public | RF_Standalone);
+			if (ConsolidateRequests.Num())
+			{
+				TSet<UObject*> ObjectsToConsolidateWithin;
+				TSet<UObject*> ObjectsToNotConsolidateWithin;
+				ObjectTools::ConsolidateObjects(ConsolidateRequests, ObjectsToConsolidateWithin, ObjectsToNotConsolidateWithin, /*bShouldDeleteAfterConsolidate=*/false);
 			}
+		}
 
+		// Move new assets into their target packages
+		for (FAssetMoveInfo& MoveInfo : AssetMoveInfos)
+		{
+			MoveInfo.NewAsset->Rename(*MoveInfo.AssetName, MoveInfo.AssetPackage, REN_DontCreateRedirectors | REN_AllowPackageLinkerMismatch);
+			MoveInfo.NewAsset->SetFlags(RF_Public | RF_Standalone);
+		}
+
+		for (UObject* AssetToSync : MergedMeshAssets)
+		{
 			AssetRegistry.AssetCreated(AssetToSync);
 			GEditor->BroadcastObjectReimported(AssetToSync);
 		}
 
-		//Also notify the content browser that the new assets exists
+		// Also notify the content browser that the new assets exists
 		FContentBrowserModule& ContentBrowserModule = FModuleManager::Get().LoadModuleChecked<FContentBrowserModule>("ContentBrowser");
-		ContentBrowserModule.Get().SyncBrowserToAssets(AssetsToSync, true);
+		ContentBrowserModule.Get().SyncBrowserToAssets(MergedMeshAssets, true);
 
 		// Place new mesh in the world
 		if (bReplaceSourceActors)
 		{
 			UStaticMesh* MergedMesh = nullptr;
-			if (AssetsToSync.FindItemByClass(&MergedMesh))
+			if (MergedMeshAssets.FindItemByClass(&MergedMesh))
 			{
 				const FScopedTransaction Transaction(LOCTEXT("PlaceMergedActor", "Place Merged Actor"));
 				UniqueLevels[0]->Modify();
@@ -232,5 +271,3 @@ bool FMeshMergingTool::RunMerge(const FString& PackageName, const TArray<TShared
 
 #undef LOCTEXT_NAMESPACE
 
-
-
